Added new DrawDebugLibrary plugin which can debug draw a whole range of different things and dispatch to multiple different drawers such as the visual logger.

diff --git a/Engine/Plugins/Experimental/DrawDebugLibrary/DrawDebugLibrary.uplugin b/Engine/Plugins/Experimental/DrawDebugLibrary/DrawDebugLibrary.uplugin
new file mode 100644
--- /dev/null
+++ b/Engine/Plugins/Experimental/DrawDebugLibrary/DrawDebugLibrary.uplugin
@@ -0,0 +1,26 @@
+{
+	"FileVersion": 3,
+	"Version": 1,
+	"VersionName": "0.1",
+	"FriendlyName": "Draw Debug Library",
+	"Description": "A library of common debug drawing functions.",
+	"Category": "Gameplay",
+	"CreatedBy": "Epic Games, Inc.",
+	"CreatedByURL": "https://epicgames.com",
+	"DocsURL": "",
+	"MarketplaceURL": "",
+	"SupportURL": "",
+	"EnabledByDefault": false,
+	"CanContainContent": true,
+	"IsBetaVersion": false,
+	"IsExperimentalVersion": true,
+	"Installed": false,
+	"Modules": [
+		{
+			"Name": "DrawDebugLibrary",
+			"Type": "Runtime",
+			"LoadingPhase": "Default"
+		},
+	],
+	"Plugins": []
+}
\ No newline at end of file
diff --git a/Engine/Plugins/Experimental/DrawDebugLibrary/Source/DrawDebugLibrary/DrawDebugLibrary.Build.cs b/Engine/Plugins/Experimental/DrawDebugLibrary/Source/DrawDebugLibrary/DrawDebugLibrary.Build.cs
new file mode 100644
--- /dev/null
+++ b/Engine/Plugins/Experimental/DrawDebugLibrary/Source/DrawDebugLibrary/DrawDebugLibrary.Build.cs
@@ -0,0 +1,53 @@
+// Copyright Epic Games, Inc. All Rights Reserved.
+
+using UnrealBuildTool;
+
+public class DrawDebugLibrary : ModuleRules
+{
+	public DrawDebugLibrary(ReadOnlyTargetRules Target) : base(Target)
+	{
+		PCHUsage = ModuleRules.PCHUsageMode.UseExplicitOrSharedPCHs;
+
+		PublicIncludePaths.AddRange(
+			new string[] {
+				// ... add public include paths required here ...
+			}
+			);
+				
+		
+		PrivateIncludePaths.AddRange(
+			new string[] {
+				// ... add other private include paths required here ...
+			}
+			);
+			
+		
+		PublicDependencyModuleNames.AddRange(
+			new string[]
+			{
+				"Core",
+				// ... add other public dependencies that you statically link with here ...
+			}
+			);
+			
+		
+		PrivateDependencyModuleNames.AddRange(
+			new string[]
+			{
+				"CoreUObject",
+				"Engine",
+				// ... add private dependencies that you statically link with here ...	
+			}
+			);
+		
+		
+		DynamicallyLoadedModuleNames.AddRange(
+			new string[]
+			{
+				// ... add any modules that your module loads dynamically here ...
+			}
+			);
+
+		bDisableAutoRTFMInstrumentation = true;
+	}
+}
diff --git a/Engine/Plugins/Experimental/DrawDebugLibrary/Source/DrawDebugLibrary/Private/DrawDebugLibrary.cpp b/Engine/Plugins/Experimental/DrawDebugLibrary/Source/DrawDebugLibrary/Private/DrawDebugLibrary.cpp
new file mode 100644
--- /dev/null
+++ b/Engine/Plugins/Experimental/DrawDebugLibrary/Source/DrawDebugLibrary/Private/DrawDebugLibrary.cpp
@@ -0,0 +1,2255 @@
+// Copyright Epic Games, Inc. All Rights Reserved.
+
+#include "DrawDebugLibrary.h"
+
+#include "Animation/TrajectoryTypes.h"
+#include "PrimitiveDrawInterface.h"
+#include "SkeletalDebugRendering.h"
+#include "HitProxies.h"
+#include "DrawDebugHelpers.h"
+#include "Animation/AnimInstanceProxy.h"
+#include "VisualLogger/VisualLogger.h"
+
+void FDrawDebugLibraryModule::StartupModule()
+{
+	// This code will execute after your module is loaded into memory; the exact timing is specified in the .uplugin file per-module
+}
+
+void FDrawDebugLibraryModule::ShutdownModule()
+{
+	// This function may be called during shutdown to clean up your module.  For modules that support dynamic reloading,
+	// we call this function before unloading the module.
+}
+	
+IMPLEMENT_MODULE(FDrawDebugLibraryModule, DrawDebugLibrary)
+
+namespace UE::DrawDebugLibrary::Private
+{
+	static FQuat FindOrientation(const FVector Start, const FVector End)
+	{
+		return (End - Start).GetSafeNormal(UE_KINDA_SMALL_NUMBER, FVector::ForwardVector).Rotation().Quaternion();
+	}
+
+	static FVector HermiteInterpolate(const FVector P0, const FVector P1, const FVector V0, const FVector V1, const float X)
+	{
+		const double W1 = 3 * X * X - 2 * X * X * X;
+		const double W2 = X * X * X - 2 * X * X + X;
+		const double W3 = X * X * X - X * X;
+
+		return W1 * (P1 - P0) + W2 * V0 + W3 * V1 + P0;
+	}
+
+	static inline double ComputeMonotoneVelocity(const double D0, const double D1)
+	{
+		return FMath::Sign(D0) != FMath::Sign(D1) ? 0.0 :
+			FMath::Clamp(
+				(D0 + D1) / 2.0,
+				FMath::Max(-3.0 * FMath::Abs(D0), -3.0 * FMath::Abs(D1)),
+				FMath::Min(+3.0 * FMath::Abs(D0), +3.0 * FMath::Abs(D1)));
+	}
+
+	static inline FVector ComputeMonotoneVelocity(const FVector D0, const FVector D1)
+	{
+		return FVector(
+			ComputeMonotoneVelocity(D0.X, D1.X),
+			ComputeMonotoneVelocity(D0.Y, D1.Y),
+			ComputeMonotoneVelocity(D0.Z, D1.Z));
+	}
+
+	static inline FVector InterpolateCubicMonoStart(const FVector P1, const FVector P2, const FVector P3, const float Alpha)
+	{
+		const FVector V1 = P2 - P1;
+		const FVector V2 = ComputeMonotoneVelocity(P2 - P1, P3 - P2);
+		return HermiteInterpolate(P1, P2, V1, V2, Alpha);
+	}
+
+	static inline FVector InterpolateCubicMonoEnd(const FVector P0, const FVector P1, const FVector P2, const float Alpha)
+	{
+		const FVector V1 = ComputeMonotoneVelocity(P1 - P0, P2 - P1);
+		const FVector V2 = P2 - P1;
+		return HermiteInterpolate(P1, P2, V1, V2, Alpha);
+	}
+
+	static inline FVector InterpolateCubicMono(const FVector P0, const FVector P1, const FVector P2, const FVector P3, const float Alpha)
+	{
+		const FVector V1 = ComputeMonotoneVelocity(P1 - P0, P2 - P1);
+		const FVector V2 = ComputeMonotoneVelocity(P2 - P1, P3 - P2);
+		return HermiteInterpolate(P1, P2, V1, V2, Alpha);
+	}
+
+	static inline FVector InterpolateCubicStart(const FVector P1, const FVector P2, const FVector P3, const float Alpha)
+	{
+		const FVector V1 = (P2 - P1);
+		const FVector V2 = ((P2 - P1) + (P3 - P2)) / 2;
+		return HermiteInterpolate(P1, P2, V1, V2, Alpha);
+	}
+
+	static inline FVector InterpolateCubicEnd(const FVector P0, const FVector P1, const FVector P2, const float Alpha)
+	{
+		const FVector V1 = ((P1 - P0) + (P2 - P1)) / 2;
+		const FVector V2 = (P2 - P1);
+		return HermiteInterpolate(P1, P2, V1, V2, Alpha);
+	}
+
+	static inline FVector InterpolateCubic(const FVector P0, const FVector P1, const FVector P2, const FVector P3, const float Alpha)
+	{
+		const FVector V1 = ((P1 - P0) + (P2 - P1)) / 2;
+		const FVector V2 = ((P2 - P1) + (P3 - P2)) / 2;
+		return HermiteInterpolate(P1, P2, V1, V2, Alpha);
+	}
+
+	// Line fonts from here: https://theorangeduck.com/page/debug-draw-text-lines
+
+	static const uint8 MonoStringLineNums[94] = {
+		 2,  6,  4, 10, 15, 16,  4,  5,  5,  3,  2,  3,  1,  1,  1, 13,  3,  6,
+		11,  4,  8, 13,  2, 14, 15,  2,  4,  2,  2,  2,  7, 28,  3, 12, 10,  9,
+		 4,  3, 11,  3,  3,  5,  4,  2,  4,  3, 12,  7, 15,  8, 12,  2,  6,  2,
+		 4,  2,  3,  3,  3,  1,  3,  2,  1,  1, 11, 10,  8, 10, 12,  5, 26,  6,
+		 4,  6,  4,  3, 11,  6, 10,  9,  9,  5, 10,  5,  6,  2,  4,  2,  4,  3,
+		10,  1, 10,  7
+	};
+
+	static const int32 MonoStringLineOffsets[94] = {
+		  0,   2,   8,  12,  22,  37,  53,  57,  62,  67,  70,  72,  75,  76,
+		 77,  78,  91,  94, 100, 111, 115, 123, 136, 138, 152, 167, 169, 173,
+		175, 177, 179, 186, 214, 217, 229, 239, 248, 252, 255, 266, 269, 272,
+		277, 281, 283, 287, 290, 302, 309, 324, 332, 344, 346, 352, 354, 358,
+		360, 363, 366, 369, 370, 373, 375, 376, 377, 388, 398, 406, 416, 428,
+		433, 459, 465, 469, 475, 479, 482, 493, 499, 509, 518, 527, 532, 542,
+		547, 553, 555, 559, 561, 565, 568, 578, 579, 589
+	};
+
+	static const int32 MonoStringLineSegments[596] = {
+		0x701F3C1F, 0x2A1F281F, 0x59135206, 0x63175913, 0x6D116317, 0x59305223,
+		0x63345930, 0x6D2E6334, 0x563A5608, 0x3B373B05, 0x6A172610, 0x26276A2E,
+		0x27222909, 0x2D2F2722, 0x35342D2F, 0x3E313534, 0x510F3E31, 0x590B510F,
+		0x620E590B, 0x681A620E, 0x6730681A, 0x19187425, 0x26077038, 0x6D096105,
+		0x70156D09, 0x671D7015, 0x5B1B671D, 0x55125B1B, 0x59085512, 0x61055908,
+		0x3D263122, 0x40323D26, 0x373A4032, 0x2B38373A, 0x252F2B38, 0x2925252F,
+		0x31222925, 0x2538590E, 0x55265F2B, 0x610D590E, 0x6A12610D, 0x6E1C6A12,
+		0x6B266E1C, 0x662A6B26, 0x5F2B662A, 0x3F085526, 0x2A282519, 0x290F2519,
+		0x2F09290F, 0x37062F09, 0x3F083706, 0x39332A28, 0x46343933, 0x6E1D6621,
+		0x5D226621, 0x541A5D22, 0x5210541A, 0x5C1B722B, 0x49165C1B, 0x37164916,
+		0x271B3716, 0x102B271B, 0x5C247214, 0x49295C24, 0x37294929, 0x27243729,
+		0x10142724, 0x70204620, 0x65324E0D, 0x650E4F32, 0x291F5A1F, 0x42074238,
+		0x2423301D, 0x191E2423, 0x1310191E, 0x422D4212, 0x2A1F2C1F, 0x70301A0C,
+		0x5A35360A, 0x4A08360A, 0x5F0D4A08, 0x66145F0D, 0x6A206614, 0x652C6A20,
+		0x5A35652C, 0x46375A35, 0x31324637, 0x2B2B3132, 0x261F2B2B, 0x2B13261F,
+		0x360A2B13, 0x2735270D, 0x6A222722, 0x5E0C6A22, 0x270B2735, 0x502F270B,
+		0x5C30502F, 0x672A5C30, 0x6B1D672A, 0x630D6B1D, 0x2625270C, 0x2F2F2625,
+		0x38332F2F, 0x40313833, 0x49224031, 0x522C4922, 0x5C2F522C, 0x652B5C2F,
+		0x6A1D652B, 0x670E6A1D, 0x49164922, 0x37053739, 0x262B6A2B, 0x6A233705,
+		0x6A2B6A23, 0x680F6830, 0x4A0F680F, 0x4A274A0F, 0x270D261F, 0x2E2E261F,
+		0x3A322E2E, 0x442F3A32, 0x4A27442F, 0x671E6930, 0x5E12671E, 0x4E0C5E12,
+		0x460B4E0C, 0x360C460B, 0x460B4D1D, 0x2914360C, 0x25212914, 0x2A2E2521,
+		0x35352A2E, 0x43333535, 0x4C294333, 0x4D1D4C29, 0x6935690A, 0x26156935,
+		0x67125C0C, 0x6B206712, 0x662E6B20, 0x5C32662E, 0x522E5C32, 0x3E0E522E,
+		0x350B3E0E, 0x2813350B, 0x261F2813, 0x2A2D261F, 0x32332A2D, 0x3C313233,
+		0x53103C31, 0x5C0C5310, 0x2920270D, 0x322D2920, 0x3E32322D, 0x4C333E32,
+		0x4C334529, 0x5F304C33, 0x67295F30, 0x6A1E6729, 0x66126A1E, 0x5E0B6612,
+		0x550A5E0B, 0x4A0D550A, 0x44144A0D, 0x421E4414, 0x4529421E, 0x531F561F,
+		0x291F2C1F, 0x531F561F, 0x23242E1E, 0x181E2324, 0x1311181E, 0x420D262E,
+		0x5C2D420D, 0x4C354C0B, 0x380A3835, 0x42312610, 0x5C114231, 0x3B1B4A1B,
+		0x4D2A4A1B, 0x59304D2A, 0x622D5930, 0x6B22622D, 0x6E146B22, 0x271B2A1B,
+		0x0F25122D, 0x0E1C0F25, 0x12100E1C, 0x1F071210, 0x34041F07, 0x4B063404,
+		0x5E0C4B06, 0x6C185E0C, 0x70246C18, 0x6D2F7024, 0x65366D2F, 0x563A6536,
+		0x483B563A, 0x3538483B, 0x2B333538, 0x292E2B33, 0x2C29292E, 0x37282C29,
+		0x2D203728, 0x4F2B3728, 0x53274F2B, 0x54225327, 0x4E195422, 0x41154E19,
+		0x32144115, 0x2B173214, 0x291B2B17, 0x2D20291B, 0x6A1F2606, 0x26396A1F,
+		0x3832380D, 0x6A0D260D, 0x672B6A0D, 0x490D4924, 0x260D292C, 0x3032292C,
+		0x39353032, 0x40323935, 0x49244032, 0x502E4924, 0x5932502E, 0x61305932,
+		0x672B6130, 0x26272A35, 0x29182627, 0x30102918, 0x390C3010, 0x460A390C,
+		0x560C460A, 0x6112560C, 0x691E6112, 0x6A27691E, 0x66356A27, 0x690A260A,
+		0x6920690A, 0x260A2620, 0x632C6920, 0x5933632C, 0x49375933, 0x39344937,
+		0x2E2E3934, 0x26202E2E, 0x260F2632, 0x690F260F, 0x6932690F, 0x490F4931,
+		0x6A102510, 0x6A326A10, 0x4810482F, 0x48354823, 0x28354835, 0x6A276635,
+		0x67196A27, 0x5D0E6719, 0x4A085D0E, 0x3C094A08, 0x300E3C09, 0x2819300E,
+		0x26252819, 0x28352625, 0x6B0A240A, 0x4936490A, 0x24366B36, 0x2633260C,
+		0x6933690C, 0x261F691F, 0x692D690E, 0x312D692D, 0x2A25312D, 0x261B2A25,
+		0x2A0D261B, 0x6A0D250D, 0x49176A32, 0x25334917, 0x490D4917, 0x26126A12,
+		0x26342612, 0x6A0B2507, 0x3F1F6A0B, 0x6A353F1F, 0x25386A35, 0x6B0C250C,
+		0x24356B0C, 0x6A352435, 0x662D6A20, 0x5E33662D, 0x4C385E33, 0x39354C38,
+		0x2C2E3935, 0x26202C2E, 0x28142620, 0x340A2814, 0x4707340A, 0x580A4707,
+		0x6613580A, 0x6A206613, 0x6A0D260D, 0x420D4226, 0x4A304226, 0x56354A30,
+		0x632F5635, 0x6A1F632F, 0x6A0D6A1F, 0x122B153A, 0x1E1F122B, 0x251F1E1F,
+		0x2B2D251F, 0x3C362B2D, 0x49383C36, 0x59354938, 0x662D5935, 0x6A20662D,
+		0x64116A20, 0x560A6411, 0x4707560A, 0x360A4707, 0x2A12360A, 0x251F2A12,
+		0x6A0D250D, 0x6A206A0D, 0x652C6A20, 0x5A31652C, 0x512F5A31, 0x512F4822,
+		0x480D4822, 0x25344822, 0x6A20682F, 0x67136A20, 0x610D6713, 0x5A0B610D,
+		0x520E5A0B, 0x4A1A520E, 0x412E4A1A, 0x3D32412E, 0x35343D32, 0x2D303534,
+		0x26222D30, 0x280A2622, 0x69396908, 0x261F691F, 0x340B6A0B, 0x6A353435,
+		0x2A10340B, 0x251F2A10, 0x292B251F, 0x3435292B, 0x251F6A06, 0x6A3A251F,
+		0x260E6A06, 0x511F260E, 0x2634511F, 0x6A382634, 0x6A352508, 0x6A0A2536,
+		0x40206A06, 0x6A394020, 0x25204020, 0x6936690B, 0x26096936, 0x26362609,
+		0x7117712C, 0x0F177117, 0x0F2C0F17, 0x1933700E, 0x71287112, 0x0F287128,
+		0x0F130F28, 0x6A1F4B0C, 0x4B336A1F, 0x0F3C0F03, 0x681D7013, 0x560F5920,
+		0x4E312531, 0x562C4E31, 0x5920562C, 0x3F153F31, 0x370D3F15, 0x2F0B370D,
+		0x26112F0B, 0x251B2611, 0x2825251B, 0x32312825, 0x290D6F0D, 0x251C290D,
+		0x2728251C, 0x36332728, 0x43353633, 0x50314335, 0x562C5031, 0x5924562C,
+		0x551B5924, 0x4B0D551B, 0x59255632, 0x55185925, 0x4A0F5518, 0x3F0C4A0F,
+		0x330F3F0C, 0x2819330F, 0x26252819, 0x29322625, 0x25317031, 0x27213731,
+		0x25192721, 0x29112519, 0x310C2911, 0x3E0A310C, 0x4E0E3E0A, 0x55154E0E,
+		0x59215515, 0x53315921, 0x4035400B, 0x4B334035, 0x532F4B33, 0x5728532F,
+		0x59205728, 0x56155920, 0x4C0D5615, 0x400B4C0D, 0x320D400B, 0x2815320D,
+		0x25242815, 0x28332524, 0x661B261B, 0x6E397027, 0x6D20661B, 0x70276D20,
+		0x4F074F37, 0x58285837, 0x512E5828, 0x591C5828, 0x5411591C, 0x4A0D5411,
+		0x3D124A0D, 0x391E3D12, 0x370E3D12, 0x310C370E, 0x2B0F310C, 0x28142B0F,
+		0x282A2814, 0x220C2814, 0x2532282A, 0x21352532, 0x1C362135, 0x16331C36,
+		0x122D1633, 0x0E1E122D, 0x110F0E1E, 0x150A110F, 0x1B09150A, 0x220C1B09,
+		0x3C2A391E, 0x47303C2A, 0x512E4730, 0x250D700D, 0x25324F32, 0x5519490D,
+		0x59235519, 0x562D5923, 0x4F32562D, 0x2634260D, 0x57212621, 0x570E5721,
+		0x6A1F6D1F, 0x6A296D29, 0x572B570D, 0x1E2B572B, 0x0E16110A, 0x11220E16,
+		0x1E2B1122, 0x250F700F, 0x411B410F, 0x5932411B, 0x2533411B, 0x2633260D,
+		0x6F212621, 0x6F0E6F21, 0x25085808, 0x25204C20, 0x25365136, 0x56134808,
+		0x59195613, 0x551E5919, 0x4C20551E, 0x572A4C20, 0x5930572A, 0x57345930,
+		0x51365734, 0x580D250D, 0x25324E32, 0x54174A0D, 0x59245417, 0x542F5924,
+		0x4E32542F, 0x500E3E09, 0x591D500E, 0x552E591D, 0x4834552E, 0x35354834,
+		0x292C3535, 0x251F292C, 0x2A11251F, 0x330B2A11, 0x3E09330B, 0x570D0E0D,
+		0x571C490D, 0x5925571C, 0x54305925, 0x43355430, 0x31324335, 0x28283132,
+		0x251D2828, 0x290D251D, 0x0F315531, 0x591F5531, 0x5312591F, 0x400A5312,
+		0x2E0D400A, 0x27132E0D, 0x251B2713, 0x2B26251B, 0x37312B26, 0x570F250F,
+		0x571F480F, 0x5928571F, 0x55325928, 0x49355532, 0x5921572F, 0x56145921,
+		0x4D0F5614, 0x44134D0F, 0x3B2B4413, 0x36303B2B, 0x30323630, 0x292C3032,
+		0x251D292C, 0x280D251D, 0x69173317, 0x2A1B3317, 0x26272A1B, 0x27342627,
+		0x58345806, 0x2F0D580D, 0x26325932, 0x28132F0D, 0x251B2813, 0x2A26251B,
+		0x34322A26, 0x251F5809, 0x5835251F, 0x25115806, 0x4A1F2511, 0x24314A1F,
+		0x58382431, 0x5833250B, 0x580D2534, 0x171A5935, 0x1113171A, 0x10071113,
+		0x241F580A, 0x5831580D, 0x260C5831, 0x2634260C, 0x6F24712F, 0x671D6F24,
+		0x4E1D671D, 0x48194E1D, 0x430C4819, 0x3F19430C, 0x371E3F19, 0x1B1E371E,
+		0x13221B1E, 0x0F2F1322, 0x0F1F7C1F, 0x6F1A710F, 0x67216F1A, 0x4E216721,
+		0x48254E21, 0x43324825, 0x3F254332, 0x37203F25, 0x1B203720, 0x131C1B20,
+		0x0F0F131C, 0x450A3D08, 0x4A13450A, 0x431E4A13, 0x3B27431E, 0x392D3B27,
+		0x3F35392D, 0x47373F35
+	};
+
+	static const char StringLineNums[94] = {
+		 5,  2,  4, 20, 21, 22,  6,  7,  7,  3,  2,  7,  1,  4,  1, 16,  3, 13,
+		14,  3, 16, 22,  2, 28, 22,  8, 11,  2,  2,  2, 17, 37,  3, 18, 17, 12,
+		 4,  3, 19,  3,  1,  9,  3,  2,  4,  3, 20, 10, 21, 11, 19,  2,  9,  2,
+		 4,  2,  3,  3,  3,  1,  3,  2,  1,  6, 14, 14, 13, 14, 16,  5, 19,  7,
+		 5,  8,  3,  1, 13,  7, 16, 14, 14,  5, 16,  5,  7,  2,  4,  2,  6,  3,
+		10,  1, 10,  8
+	};
+
+	static const char StringLineAdvs[94] = {
+		40, 64, 84, 80, 96, 96, 40, 48, 48, 64, 88, 40, 88, 40, 64, 80, 80, 80,
+		80, 80, 80, 80, 80, 80, 80, 40, 40, 96, 104, 96, 72, 108, 72, 84, 84, 84,
+		76, 72, 84, 88, 32, 64, 84, 68, 96, 88, 88, 84, 88, 84, 80, 64, 88, 72,
+		96, 80, 72, 80, 44, 56, 44, 64, 64, 40, 76, 76, 72, 76, 72, 48, 76, 76,
+		32, 40, 68, 32, 120, 76, 76, 76, 76, 52, 68, 48, 76, 64, 88, 68, 64, 68,
+		56, 32, 56, 96
+	};
+
+	static const int StringLineOffsets[94] = {
+		  0,   5,   7,  11,  31,  52,  74,  80,  87,  94,  97,  99, 106, 107,
+		111, 112, 128, 131, 144, 158, 161, 177, 199, 201, 229, 251, 259, 270,
+		272, 274, 276, 293, 330, 333, 351, 368, 380, 384, 387, 406, 409, 410,
+		419, 422, 424, 428, 431, 451, 461, 482, 493, 512, 514, 523, 525, 529,
+		531, 534, 537, 540, 541, 544, 546, 547, 553, 567, 581, 594, 608, 624,
+		629, 648, 655, 660, 668, 671, 672, 685, 692, 708, 722, 736, 741, 757,
+		762, 769, 771, 775, 777, 783, 786, 796, 797, 807
+	};
+
+	static const int StringLineSegments[815] = {
+		0x3C147414, 0x24102814, 0x20142410, 0x24182014, 0x28142418, 0x58107410,
+		0x58307430, 0x18187024, 0x1830703C, 0x50485010, 0x3844380C, 0x10248430,
+		0x703C6844, 0x7430703C, 0x74207430, 0x70147420, 0x680C7014, 0x600C680C,
+		0x5810600C, 0x54145810, 0x501C5414, 0x4834501C, 0x443C4834, 0x4040443C,
+		0x38444040, 0x2C443844, 0x243C2C44, 0x2030243C, 0x20202030, 0x24142020,
+		0x2C0C2414, 0x201C7444, 0x6C287420, 0x64286C28, 0x5C246428, 0x581C5C24,
+		0x5814581C, 0x600C5814, 0x680C600C, 0x7010680C, 0x74187010, 0x74207418,
+		0x383C3C44, 0x3038383C, 0x28383038, 0x20402838, 0x20482040, 0x24502048,
+		0x2C542450, 0x34542C54, 0x3C4C3454, 0x3C443C4C, 0x2C3C404C, 0x24342C3C,
+		0x202C2434, 0x201C202C, 0x2414201C, 0x28102414, 0x300C2810, 0x380C300C,
+		0x4010380C, 0x44144010, 0x54304414, 0x58345430, 0x60385834, 0x68386038,
+		0x70346838, 0x742C7034, 0x7024742C, 0x68207024, 0x60206820, 0x54246020,
+		0x482C5424, 0x204C482C, 0x70106C14, 0x74147010, 0x70187414, 0x68187018,
+		0x60146818, 0x5C106014, 0x701C7820, 0x6014701C, 0x4C106014, 0x3C104C10,
+		0x28143C10, 0x181C2814, 0x1020181C, 0x70147810, 0x601C7014, 0x4C20601C,
+		0x3C204C20, 0x281C3C20, 0x1814281C, 0x10101814, 0x44207420, 0x5034680C,
+		0x500C6834, 0x282C602C, 0x44484410, 0x20142418, 0x24102014, 0x28142410,
+		0x24182814, 0x1C182418, 0x14141C18, 0x10101414, 0x44484410, 0x24102814,
+		0x20142410, 0x24182014, 0x28142418, 0x20087430, 0x70187424, 0x64107018,
+		0x500C6410, 0x440C500C, 0x3010440C, 0x24183010, 0x20242418, 0x202C2024,
+		0x2438202C, 0x30402438, 0x44443040, 0x50444444, 0x64405044, 0x70386440,
+		0x742C7038, 0x7424742C, 0x68206418, 0x742C6820, 0x202C742C, 0x64106010,
+		0x6C146410, 0x70186C14, 0x74207018, 0x74307420, 0x70387430, 0x6C3C7038,
+		0x64406C3C, 0x5C406440, 0x543C5C40, 0x4834543C, 0x200C4834, 0x2044200C,
+		0x74407414, 0x54287440, 0x54345428, 0x503C5434, 0x4C40503C, 0x40444C40,
+		0x38444044, 0x2C403844, 0x24382C40, 0x202C2438, 0x2020202C, 0x24142020,
+		0x28102414, 0x300C2810, 0x3C0C7434, 0x3C483C0C, 0x20347434, 0x7414743C,
+		0x50107414, 0x54145010, 0x58205414, 0x582C5820, 0x5438582C, 0x4C405438,
+		0x40444C40, 0x38444044, 0x2C403844, 0x24382C40, 0x202C2438, 0x2020202C,
+		0x24142020, 0x28102414, 0x300C2810, 0x703C6840, 0x7430703C, 0x74287430,
+		0x701C7428, 0x6414701C, 0x50106414, 0x3C105010, 0x2C143C10, 0x241C2C14,
+		0x2028241C, 0x202C2028, 0x2438202C, 0x2C402438, 0x38442C40, 0x3C443844,
+		0x48403C44, 0x50384840, 0x542C5038, 0x5428542C, 0x501C5428, 0x4814501C,
+		0x3C104814, 0x201C7444, 0x7444740C, 0x70147420, 0x68107014, 0x60106810,
+		0x58146010, 0x541C5814, 0x502C541C, 0x4C38502C, 0x44404C38, 0x3C444440,
+		0x30443C44, 0x28403044, 0x243C2840, 0x2030243C, 0x20202030, 0x24142020,
+		0x28102414, 0x300C2810, 0x3C0C300C, 0x44103C0C, 0x4C184410, 0x50244C18,
+		0x54345024, 0x583C5434, 0x6040583C, 0x68406040, 0x703C6840, 0x7430703C,
+		0x74207430, 0x4C3C5840, 0x44344C3C, 0x40284434, 0x40244028, 0x44184024,
+		0x4C104418, 0x580C4C10, 0x5C0C580C, 0x68105C0C, 0x70186810, 0x74247018,
+		0x74287424, 0x70347428, 0x683C7034, 0x5840683C, 0x44405840, 0x303C4440,
+		0x2434303C, 0x20282434, 0x20202028, 0x24142020, 0x2C102414, 0x54105814,
+		0x50145410, 0x54185014, 0x58145418, 0x24102814, 0x20142410, 0x24182014,
+		0x28142418, 0x54105814, 0x50145410, 0x54185014, 0x58145418, 0x20142418,
+		0x24102014, 0x28142410, 0x24182814, 0x1C182418, 0x14141C18, 0x10101414,
+		0x44106850, 0x20504410, 0x50585010, 0x38583810, 0x44506810, 0x20104450,
+		0x640C600C, 0x6C10640C, 0x70146C10, 0x741C7014, 0x742C741C, 0x7034742C,
+		0x6C387034, 0x643C6C38, 0x5C3C643C, 0x54385C3C, 0x50345438, 0x48245034,
+		0x3C244824, 0x24202824, 0x20242420, 0x24282024, 0x28242428, 0x5C445448,
+		0x603C5C44, 0x6030603C, 0x5C286030, 0x58245C28, 0x4C205824, 0x40204C20,
+		0x38244020, 0x342C3824, 0x3438342C, 0x38403438, 0x40443840, 0x40446048,
+		0x38444044, 0x344C3844, 0x3454344C, 0x3C5C3454, 0x48603C5C, 0x50604860,
+		0x5C5C5060, 0x64585C5C, 0x6C506458, 0x70486C50, 0x743C7048, 0x7430743C,
+		0x70247430, 0x6C1C7024, 0x64146C1C, 0x5C106414, 0x500C5C10, 0x440C500C,
+		0x3810440C, 0x30143810, 0x281C3014, 0x2424281C, 0x20302424, 0x20402030,
+		0x20047424, 0x20447424, 0x3C383C10, 0x20107410, 0x74347410, 0x70407434,
+		0x6C447040, 0x64486C44, 0x5C486448, 0x54445C48, 0x50405444, 0x4C345040,
+		0x4C344C10, 0x48404C34, 0x44444840, 0x3C484444, 0x30483C48, 0x28443048,
+		0x24402844, 0x20342440, 0x20102034, 0x68446048, 0x703C6844, 0x7434703C,
+		0x74247434, 0x701C7424, 0x6814701C, 0x60106814, 0x540C6010, 0x400C540C,
+		0x3410400C, 0x2C143410, 0x241C2C14, 0x2024241C, 0x20342024, 0x243C2034,
+		0x2C44243C, 0x34482C44, 0x20107410, 0x742C7410, 0x7038742C, 0x68407038,
+		0x60446840, 0x54486044, 0x40485448, 0x34444048, 0x2C403444, 0x24382C40,
+		0x202C2438, 0x2010202C, 0x20107410, 0x74447410, 0x4C304C10, 0x20442010,
+		0x20107410, 0x74447410, 0x4C304C10, 0x68446048, 0x703C6844, 0x7434703C,
+		0x74247434, 0x701C7424, 0x6814701C, 0x60106814, 0x540C6010, 0x400C540C,
+		0x3410400C, 0x2C143410, 0x241C2C14, 0x2024241C, 0x20342024, 0x243C2034,
+		0x2C44243C, 0x34482C44, 0x40483448, 0x40484034, 0x20107410, 0x20487448,
+		0x4C484C10, 0x20107410, 0x34307430, 0x282C3430, 0x2428282C, 0x20202428,
+		0x20182020, 0x24102018, 0x280C2410, 0x3408280C, 0x3C083408, 0x20107410,
+		0x3C107448, 0x20485024, 0x20107410, 0x20402010, 0x20107410, 0x20307410,
+		0x20307450, 0x20507450, 0x20107410, 0x20487410, 0x20487448, 0x701C7424,
+		0x6814701C, 0x60106814, 0x540C6010, 0x400C540C, 0x3410400C, 0x2C143410,
+		0x241C2C14, 0x2024241C, 0x20342024, 0x243C2034, 0x2C44243C, 0x34482C44,
+		0x404C3448, 0x544C404C, 0x6048544C, 0x68446048, 0x703C6844, 0x7434703C,
+		0x74247434, 0x20107410, 0x74347410, 0x70407434, 0x6C447040, 0x64486C44,
+		0x58486448, 0x50445848, 0x4C405044, 0x48344C40, 0x48104834, 0x701C7424,
+		0x6814701C, 0x60106814, 0x540C6010, 0x400C540C, 0x3410400C, 0x2C143410,
+		0x241C2C14, 0x2024241C, 0x20342024, 0x243C2034, 0x2C44243C, 0x34482C44,
+		0x404C3448, 0x544C404C, 0x6048544C, 0x68446048, 0x703C6844, 0x7434703C,
+		0x74247434, 0x18483030, 0x20107410, 0x74347410, 0x70407434, 0x6C447040,
+		0x64486C44, 0x5C486448, 0x54445C48, 0x50405444, 0x4C345040, 0x4C104C34,
+		0x20484C2C, 0x703C6844, 0x7430703C, 0x74207430, 0x70147420, 0x680C7014,
+		0x600C680C, 0x5810600C, 0x54145810, 0x501C5414, 0x4834501C, 0x443C4834,
+		0x4040443C, 0x38444040, 0x2C443844, 0x243C2C44, 0x2030243C, 0x20202030,
+		0x24142020, 0x2C0C2414, 0x20207420, 0x743C7404, 0x38107410, 0x2C143810,
+		0x241C2C14, 0x2028241C, 0x20302028, 0x243C2030, 0x2C44243C, 0x38482C44,
+		0x74483848, 0x20247404, 0x20247444, 0x201C7408, 0x201C7430, 0x20447430,
+		0x20447458, 0x2044740C, 0x200C7444, 0x4C247404, 0x20244C24, 0x4C247444,
+		0x200C7444, 0x7444740C, 0x2044200C, 0x78107820, 0x18107810, 0x18201810,
+		0x20307408, 0x78207810, 0x18207820, 0x18101820, 0x70205C0C, 0x5C347020,
+		0x18401800, 0x70147418, 0x68107014, 0x60106810, 0x5C146010, 0x60185C14,
+		0x64146018, 0x203C583C, 0x54344C3C, 0x582C5434, 0x5820582C, 0x54185820,
+		0x4C105418, 0x400C4C10, 0x380C400C, 0x2C10380C, 0x24182C10, 0x20202418,
+		0x202C2020, 0x2434202C, 0x2C3C2434, 0x20107410, 0x54184C10, 0x58205418,
+		0x582C5820, 0x5434582C, 0x4C3C5434, 0x40404C3C, 0x38404040, 0x2C3C3840,
+		0x24342C3C, 0x202C2434, 0x2020202C, 0x24182020, 0x2C102418, 0x54344C3C,
+		0x582C5434, 0x5820582C, 0x54185820, 0x4C105418, 0x400C4C10, 0x380C400C,
+		0x2C10380C, 0x24182C10, 0x20202418, 0x202C2020, 0x2434202C, 0x2C3C2434,
+		0x203C743C, 0x54344C3C, 0x582C5434, 0x5820582C, 0x54185820, 0x4C105418,
+		0x400C4C10, 0x380C400C, 0x2C10380C, 0x24182C10, 0x20202418, 0x202C2020,
+		0x2434202C, 0x2C3C2434, 0x403C400C, 0x483C403C, 0x5038483C, 0x54345038,
+		0x582C5434, 0x5820582C, 0x54185820, 0x4C105418, 0x400C4C10, 0x380C400C,
+		0x2C10380C, 0x24182C10, 0x20202418, 0x202C2020, 0x2434202C, 0x2C3C2434,
+		0x74207428, 0x70187420, 0x64147018, 0x20146414, 0x58245808, 0x183C583C,
+		0x0C38183C, 0x08340C38, 0x042C0834, 0x0420042C, 0x08180420, 0x54344C3C,
+		0x582C5434, 0x5820582C, 0x54185820, 0x4C105418, 0x400C4C10, 0x380C400C,
+		0x2C10380C, 0x24182C10, 0x20202418, 0x202C2020, 0x2434202C, 0x2C3C2434,
+		0x20107410, 0x541C4810, 0x5824541C, 0x58305824, 0x54385830, 0x483C5438,
+		0x203C483C, 0x7010740C, 0x74147010, 0x78107414, 0x740C7810, 0x20105810,
+		0x70187414, 0x741C7018, 0x7818741C, 0x74147818, 0x14185818, 0x08141418,
+		0x040C0814, 0x0404040C, 0x20107410, 0x30105838, 0x203C4020, 0x20107410,
+		0x20105810, 0x541C4810, 0x5824541C, 0x58305824, 0x54385830, 0x483C5438,
+		0x203C483C, 0x5448483C, 0x58505448, 0x585C5850, 0x5464585C, 0x48685464,
+		0x20684868, 0x20105810, 0x541C4810, 0x5824541C, 0x58305824, 0x54385830,
+		0x483C5438, 0x203C483C, 0x54185820, 0x4C105418, 0x400C4C10, 0x380C400C,
+		0x2C10380C, 0x24182C10, 0x20202418, 0x202C2020, 0x2434202C, 0x2C3C2434,
+		0x38402C3C, 0x40403840, 0x4C3C4040, 0x54344C3C, 0x582C5434, 0x5820582C,
+		0x04105810, 0x54184C10, 0x58205418, 0x582C5820, 0x5434582C, 0x4C3C5434,
+		0x40404C3C, 0x38404040, 0x2C3C3840, 0x24342C3C, 0x202C2434, 0x2020202C,
+		0x24182020, 0x2C102418, 0x043C583C, 0x54344C3C, 0x582C5434, 0x5820582C,
+		0x54185820, 0x4C105418, 0x400C4C10, 0x380C400C, 0x2C10380C, 0x24182C10,
+		0x20202418, 0x202C2020, 0x2434202C, 0x2C3C2434, 0x20105810, 0x4C144010,
+		0x541C4C14, 0x5824541C, 0x58305824, 0x54344C38, 0x58285434, 0x581C5828,
+		0x5410581C, 0x4C0C5410, 0x44104C0C, 0x40184410, 0x3C2C4018, 0x38343C2C,
+		0x30383834, 0x2C383038, 0x24342C38, 0x20282434, 0x201C2028, 0x2410201C,
+		0x2C0C2410, 0x30147414, 0x24183014, 0x20202418, 0x20282020, 0x58245808,
+		0x30105810, 0x24143010, 0x201C2414, 0x2028201C, 0x24302028, 0x303C2430,
+		0x203C583C, 0x20205808, 0x20205838, 0x201C580C, 0x201C582C, 0x203C582C,
+		0x203C584C, 0x2038580C, 0x200C5838, 0x20205808, 0x20205838, 0x10182020,
+		0x08101018, 0x04080810, 0x04040408, 0x200C5838, 0x5838580C, 0x2038200C,
+		0x741C7824, 0x6C18741C, 0x50186C18, 0x48145018, 0x44104814, 0x40144410,
+		0x38184014, 0x1C183818, 0x141C1C18, 0x1024141C, 0x10107810, 0x74187810,
+		0x6C1C7418, 0x501C6C1C, 0x4820501C, 0x44244820, 0x40204424, 0x381C4020,
+		0x1C1C381C, 0x14181C1C, 0x10101418, 0x4810400C, 0x4C184810, 0x4C204C18,
+		0x40384C20, 0x3C404038, 0x3C483C40, 0x40503C48, 0x48544050
+	};
+
+	static ELogVerbosity::Type LogVerbosityFromDrawDebugLogVerbosity(const EDrawDebugLogVerbosity Verbosity)
+	{
+		switch (Verbosity)
+		{
+		case EDrawDebugLogVerbosity::Fatal: return ELogVerbosity::Fatal;
+		case EDrawDebugLogVerbosity::Error: return ELogVerbosity::Error;
+		case EDrawDebugLogVerbosity::Warning: return ELogVerbosity::Warning;
+		case EDrawDebugLogVerbosity::Display: return ELogVerbosity::Display;
+		case EDrawDebugLogVerbosity::Log: return ELogVerbosity::Log;
+		case EDrawDebugLogVerbosity::Verbose: return ELogVerbosity::Verbose;
+		case EDrawDebugLogVerbosity::VeryVerbose: return ELogVerbosity::VeryVerbose;
+		default: checkNoEntry(); return ELogVerbosity::Log;
+		}
+	}
+}
+
+FDebugDrawer UDrawDebugLibrary::MakeNullDebugDrawer()
+{
+	return FDebugDrawer();
+}
+
+FDebugDrawer UDrawDebugLibrary::MakeDebugDrawer(UObject* Object)
+{
+	return MakeObjectDebugDrawer(Object);
+}
+
+FDebugDrawer UDrawDebugLibrary::MakeObjectDebugDrawer(UObject* Object)
+{
+	return MakeWorldDebugDrawer(Object ? Object->GetWorld() : nullptr);
+}
+
+FDebugDrawer UDrawDebugLibrary::MakeWorldDebugDrawer(UWorld* World)
+{
+	FDebugDrawer Drawer;
+	Drawer.World = World;
+	return Drawer;
+}
+
+FDebugDrawer UDrawDebugLibrary::MakeAnimInstanceProxyDebugDrawer(FAnimInstanceProxy* AnimInstanceProxy)
+{
+	FDebugDrawer Drawer;
+	Drawer.AnimInstanceProxy = AnimInstanceProxy;
+	return Drawer;
+}
+
+FDebugDrawer UDrawDebugLibrary::MakePDIDebugDrawer(FPrimitiveDrawInterface* PrimitiveDrawInterface)
+{
+	FDebugDrawer Drawer;
+	Drawer.PDI = PrimitiveDrawInterface;
+	return Drawer;
+}
+
+FDebugDrawer UDrawDebugLibrary::MakeVisualLoggerDebugDrawer(UObject* Object, const FName Category, const EDrawDebugLogVerbosity Verbosity, const bool bDrawToScene)
+{
+	return MakeVisualLoggerDebugDrawerFromObject(Object, Category, Verbosity, bDrawToScene);
+}
+
+FDebugDrawer UDrawDebugLibrary::MakeVisualLoggerDebugDrawerFromObject(UObject* Object, const FName Category, const EDrawDebugLogVerbosity Verbosity, const bool bDrawToScene)
+{
+	FDebugDrawer Drawer;
+	Drawer.World = (Object && bDrawToScene) ? Object->GetWorld() : nullptr;
+	Drawer.VisualLoggerObject = Object;
+	Drawer.VisualLoggerCategory = Category;
+	Drawer.VisualLoggerVerbosity = UE::DrawDebugLibrary::Private::LogVerbosityFromDrawDebugLogVerbosity(Verbosity);
+	return Drawer;
+}
+
+FDebugDrawer UDrawDebugLibrary::MakeVisualLoggerDebugDrawerFromObjectWithCategory(UObject* Object, const FLogCategoryBase& Category, const ELogVerbosity::Type Verbosity, const bool bDrawToScene)
+{
+	FDebugDrawer Drawer;
+	Drawer.World = (Object && bDrawToScene) ? Object->GetWorld() : nullptr;
+	Drawer.VisualLoggerObject = Object;
+	Drawer.VisualLoggerCategory = Category.GetCategoryName();
+	Drawer.VisualLoggerVerbosity = Verbosity;
+	return Drawer;
+}
+
+FDebugDrawer UDrawDebugLibrary::MakeMergedDebugDrawer(const TArray<FDebugDrawer>& DebugDrawers)
+{
+	return MakeMergedDebugDrawerArrayView(DebugDrawers);
+}
+
+FDebugDrawer UDrawDebugLibrary::MakeMergedDebugDrawerArrayView(const TArrayView<const FDebugDrawer> DebugDrawers)
+{
+	FDebugDrawer Drawer;
+
+	for (const FDebugDrawer& SubDrawer : DebugDrawers)
+	{
+		if (SubDrawer.PDI) { Drawer.PDI = SubDrawer.PDI; }
+		if (SubDrawer.World) { Drawer.World = SubDrawer.World; }
+		if (SubDrawer.AnimInstanceProxy) { Drawer.AnimInstanceProxy = SubDrawer.AnimInstanceProxy; }
+		if (SubDrawer.VisualLoggerObject)
+		{
+			Drawer.VisualLoggerObject = SubDrawer.VisualLoggerObject;
+			Drawer.VisualLoggerCategory = SubDrawer.VisualLoggerCategory;
+			Drawer.VisualLoggerVerbosity = SubDrawer.VisualLoggerVerbosity;
+		}
+	}
+
+	return Drawer;
+}
+
+void UDrawDebugLibrary::LogVisualLoggerString(const FDebugDrawer& Drawer, const FString& String, const FLinearColor Color)
+{
+	LogVisualLoggerStringView(Drawer, String, Color);
+}
+
+void UDrawDebugLibrary::LogVisualLoggerStringView(const FDebugDrawer& Drawer, const FStringView String, const FLinearColor Color)
+{
+#if ENABLE_VISUAL_LOG
+	if (Drawer.VisualLoggerObject)
+	{
+		FVisualLogEntry* CurrentEntry = nullptr;
+		if (FVisualLogger::CheckVisualLogInputInternal(Drawer.VisualLoggerObject, Drawer.VisualLoggerCategory, Drawer.VisualLoggerVerbosity, /*OutWorld=*/ nullptr, &CurrentEntry))
+		{
+			CurrentEntry->LogLines.Add(FVisualLogLine(Drawer.VisualLoggerCategory, Drawer.VisualLoggerVerbosity, FString(String)));
+		}
+	}
+#endif
+}
+
+void UDrawDebugLibrary::LogVisualLoggerText(const FDebugDrawer& Drawer, const FText& Text, const FLinearColor Color)
+{
+	LogVisualLoggerStringView(Drawer, Text.ToString(), Color);
+}
+
+void UDrawDebugLibrary::MakeLinearlySpacedFloatArray(TArray<float>& OutValues, const float Start, const float Stop, const int32 Num)
+{
+	OutValues.SetNumUninitialized(Num);
+
+	if (Num == 1)
+	{
+		OutValues[0] = Start;
+		return;
+	}
+
+	for (int32 Idx = 0; Idx < Num; Idx++)
+	{
+		OutValues[Idx] = Start + (Stop - Start) * ((float)(Idx) / (Num - 1));
+	}
+}
+
+void UDrawDebugLibrary::AddToFloatHistoryArray(TArray<float>& InOutValues, const float NewValue, const int32 MaxHistoryNum)
+{
+	if (MaxHistoryNum <= 0)
+	{
+		InOutValues.Empty();
+		return;
+	}
+
+	InOutValues.Add(NewValue);
+
+	while (InOutValues.Num() > MaxHistoryNum)
+	{
+		// Shift Values
+		for (int32 Idx = 0; Idx < InOutValues.Num() - 1; Idx++)
+		{
+			InOutValues[Idx] = InOutValues[Idx + 1];
+		}
+
+		// Pop End
+		InOutValues.Pop();
+	}
+}
+
+void UDrawDebugLibrary::AddToVectorHistoryArray(TArray<FVector>& InOutValues, const FVector NewValue, const int32 MaxHistoryNum)
+{
+	if (MaxHistoryNum <= 0)
+	{
+		InOutValues.Empty();
+		return;
+	}
+
+	InOutValues.Add(NewValue);
+
+	while (InOutValues.Num() > MaxHistoryNum)
+	{
+		// Shift Values
+		for (int32 Idx = 0; Idx < InOutValues.Num() - 1; Idx++)
+		{
+			InOutValues[Idx] = InOutValues[Idx + 1];
+		}
+
+		// Pop End
+		InOutValues.Pop();
+	}
+}
+
+void UDrawDebugLibrary::DrawDebugPoint(const FDebugDrawer& Drawer, const FVector& Location, const FLinearColor Color, const float Size, const bool bDepthTest)
+{
+	if (Drawer.PDI)
+	{
+		Drawer.PDI->DrawPoint(Location, Color, Size, bDepthTest ? SDPG_World : SDPG_Foreground);
+	}
+
+#if UE_ENABLE_DEBUG_DRAWING
+	if (Drawer.World)
+	{
+		::DrawDebugPoint(Drawer.World, Location, Size, Color.ToFColor(true), false, -1.0f, bDepthTest ? 0 : 1);
+	}
+
+	if (Drawer.AnimInstanceProxy)
+	{
+		Drawer.AnimInstanceProxy->AnimDrawDebugPoint(Location, Size, Color.ToFColor(true), false, -1.0f, bDepthTest ? SDPG_World : SDPG_Foreground);
+	}
+#endif
+
+#if ENABLE_VISUAL_LOG
+	if (Drawer.VisualLoggerObject)
+	{
+		FVisualLogEntry* CurrentEntry = nullptr;
+		if (FVisualLogger::CheckVisualLogInputInternal(Drawer.VisualLoggerObject, Drawer.VisualLoggerCategory, Drawer.VisualLoggerVerbosity, /*OutWorld=*/ nullptr, &CurrentEntry))
+		{
+			const uint16 IntegerSize = Size == 0.0f ? 0 : FMath::CeilToInt(FMath::Max(2.0f * Size, 0.0f));
+			FVisualLogShapeElement Element(TEXT(""), Color.ToFColor(false), IntegerSize, Drawer.VisualLoggerCategory);
+			Element.Type = EVisualLoggerShapeElement::SinglePoint;
+			Element.Radius = 0.0f;
+			Element.Verbosity = Drawer.VisualLoggerVerbosity;
+			Element.Points.Add(Location);
+			CurrentEntry->ElementsToDraw.Add(Element);
+		}
+	}
+#endif
+}
+
+void UDrawDebugLibrary::DrawDebugPoints(const FDebugDrawer& Drawer, const TArray<FVector>& Locations, const FLinearColor Color, const float Size, const bool bDepthTest)
+{
+	DrawDebugPointsArrayView(Drawer, Locations, Color, Size, bDepthTest);
+}
+
+void UDrawDebugLibrary::DrawDebugPointsArrayView(const FDebugDrawer& Drawer, const TArrayView<const FVector> Locations, const FLinearColor Color, const float Size, const bool bDepthTest)
+{
+	const int32 Num = Locations.Num();
+
+	if (Drawer.PDI)
+	{
+		for (int32 Idx = 0; Idx < Num; Idx++)
+		{
+			Drawer.PDI->DrawPoint(Locations[Idx], Color, Size, bDepthTest ? SDPG_World : SDPG_Foreground);
+		}
+	}
+
+#if UE_ENABLE_DEBUG_DRAWING
+	if (Drawer.World)
+	{
+		for (int32 Idx = 0; Idx < Num; Idx++)
+		{
+			::DrawDebugPoint(Drawer.World, Locations[Idx], Size, Color.ToFColor(true), false, -1.0f, bDepthTest ? 0 : 1);
+		}
+	}
+
+	if (Drawer.AnimInstanceProxy)
+	{
+		for (int32 Idx = 0; Idx < Num; Idx++)
+		{
+			Drawer.AnimInstanceProxy->AnimDrawDebugPoint(Locations[Idx], Size, Color.ToFColor(true), false, -1.0f, bDepthTest ? SDPG_World : SDPG_Foreground);
+		}
+	}
+#endif
+
+#if ENABLE_VISUAL_LOG
+	if (Drawer.VisualLoggerObject)
+	{
+		FVisualLogEntry* CurrentEntry = nullptr;
+		if (FVisualLogger::CheckVisualLogInputInternal(Drawer.VisualLoggerObject, Drawer.VisualLoggerCategory, Drawer.VisualLoggerVerbosity, /*OutWorld=*/ nullptr, &CurrentEntry))
+		{
+			const uint16 IntegerSize = Size == 0.0f ? 0 : FMath::CeilToInt(FMath::Max(2.0f * Size, 0.0f));
+			FVisualLogShapeElement Element(TEXT(""), Color.ToFColor(false), IntegerSize, Drawer.VisualLoggerCategory);
+			Element.Type = EVisualLoggerShapeElement::SinglePoint;
+			Element.Radius = 0.0f;
+			Element.Verbosity = Drawer.VisualLoggerVerbosity;
+			Element.Points.Append(Locations);
+			CurrentEntry->ElementsToDraw.Add(Element);
+		}
+	}
+#endif
+}
+
+void UDrawDebugLibrary::DrawDebugLine(const FDebugDrawer& Drawer, const FVector& StartLocation, const FVector& EndLocation, const FLinearColor Color, const float Thickness, const bool bDepthTest)
+{
+	if (Drawer.PDI)
+	{
+		Drawer.PDI->DrawTranslucentLine(StartLocation, EndLocation, Color, SDPG_World, Thickness, bDepthTest ? 0.0f : UE_MAX_FLT);
+	}
+	
+#if UE_ENABLE_DEBUG_DRAWING
+	if (Drawer.World)
+	{
+		::DrawDebugLine(Drawer.World, StartLocation, EndLocation, Color.ToFColor(true), false, -1.0f, bDepthTest ? 0 : 1, Thickness);
+	}
+
+	if (Drawer.AnimInstanceProxy)
+	{
+		Drawer.AnimInstanceProxy->AnimDrawDebugLine(StartLocation, EndLocation, Color.ToFColor(true), false, -1.0f, Thickness, bDepthTest ? SDPG_World : SDPG_Foreground);
+	}
+#endif
+
+#if ENABLE_VISUAL_LOG
+	if (Drawer.VisualLoggerObject)
+	{
+		FVisualLogEntry* CurrentEntry = nullptr;
+		if (FVisualLogger::CheckVisualLogInputInternal(Drawer.VisualLoggerObject, Drawer.VisualLoggerCategory, Drawer.VisualLoggerVerbosity, /*OutWorld=*/ nullptr, &CurrentEntry))
+		{
+			const uint16 IntegerThickness = Thickness == 0.0f ? 0 : FMath::CeilToInt(FMath::Max(2.0f * Thickness, 0.0f));
+			FVisualLogShapeElement Element(TEXT(""), Color.ToFColor(false), IntegerThickness, Drawer.VisualLoggerCategory);
+			Element.Type = EVisualLoggerShapeElement::Segment;
+			Element.Verbosity = Drawer.VisualLoggerVerbosity;
+			Element.Points.Reserve(2);
+			Element.Points.Add(StartLocation);
+			Element.Points.Add(EndLocation);
+			CurrentEntry->ElementsToDraw.Add(Element);
+		}
+	}
+#endif
+}
+
+void UDrawDebugLibrary::DrawDebugLines(const FDebugDrawer& Drawer, const TArray<FVector>& StartLocations, const TArray<FVector>& EndLocations, const FLinearColor Color, const float Thickness, const bool bDepthTest)
+{
+	DrawDebugLinesArrayView(Drawer, StartLocations, EndLocations, Color, Thickness, bDepthTest);
+}
+
+void UDrawDebugLibrary::DrawDebugLinesArrayView(const FDebugDrawer& Drawer, const TArrayView<const FVector> StartLocations, const TArrayView<const FVector> EndLocations, const FLinearColor Color, const float Thickness, const bool bDepthTest)
+{
+	const int32 Num = FMath::Min(StartLocations.Num(), EndLocations.Num());
+
+	if (Drawer.PDI)
+	{
+		for (int32 Idx = 0; Idx < Num; Idx++)
+		{
+			Drawer.PDI->DrawTranslucentLine(StartLocations[Idx], EndLocations[Idx], Color, SDPG_World, Thickness, bDepthTest ? 0.0f : UE_MAX_FLT);
+		}
+	}
+
+#if UE_ENABLE_DEBUG_DRAWING
+	if (Drawer.World)
+	{
+		for (int32 Idx = 0; Idx < Num; Idx++)
+		{
+			::DrawDebugLine(Drawer.World, StartLocations[Idx], EndLocations[Idx], Color.ToFColor(true), false, -1.0f, bDepthTest ? 0 : 1, Thickness);
+		}
+	}
+
+	if (Drawer.AnimInstanceProxy)
+	{
+		for (int32 Idx = 0; Idx < Num; Idx++)
+		{
+			Drawer.AnimInstanceProxy->AnimDrawDebugLine(StartLocations[Idx], EndLocations[Idx], Color.ToFColor(true), false, -1.0f, Thickness, bDepthTest ? SDPG_World : SDPG_Foreground);
+		}
+	}
+#endif
+
+#if ENABLE_VISUAL_LOG
+	if (Drawer.VisualLoggerObject)
+	{
+		FVisualLogEntry* CurrentEntry = nullptr;
+		if (FVisualLogger::CheckVisualLogInputInternal(Drawer.VisualLoggerObject, Drawer.VisualLoggerCategory, Drawer.VisualLoggerVerbosity, /*OutWorld=*/ nullptr, &CurrentEntry))
+		{
+			const uint16 IntegerThickness = Thickness == 0.0f ? 0 : FMath::CeilToInt(FMath::Max(2.0f * Thickness, 0.0f));
+			FVisualLogShapeElement Element(TEXT(""), Color.ToFColor(false), IntegerThickness, Drawer.VisualLoggerCategory);
+			Element.Type = EVisualLoggerShapeElement::Segment;
+			Element.Verbosity = Drawer.VisualLoggerVerbosity;
+			Element.Points.Reserve(2 * Num);
+			for (int32 Idx = 0; Idx < Num; Idx++)
+			{
+				Element.Points.Add(StartLocations[Idx]);
+				Element.Points.Add(EndLocations[Idx]);
+			}
+			CurrentEntry->ElementsToDraw.Add(Element);
+		}
+	}
+#endif
+}
+
+void UDrawDebugLibrary::DrawDebugLinesPairsArrayView(const FDebugDrawer& Drawer, const TArrayView<const TPair<FVector, FVector>> Segments, const FLinearColor Color, const float Thickness, const bool bDepthTest)
+{
+	const int32 Num = Segments.Num();
+
+	if (Drawer.PDI)
+	{
+		for (int32 Idx = 0; Idx < Num; Idx++)
+		{
+			Drawer.PDI->DrawTranslucentLine(Segments[Idx].Key, Segments[Idx].Value, Color, SDPG_World, Thickness, bDepthTest ? 0.0f : UE_MAX_FLT);
+		}
+	}
+
+#if UE_ENABLE_DEBUG_DRAWING
+	if (Drawer.World)
+	{
+		for (int32 Idx = 0; Idx < Num; Idx++)
+		{
+			::DrawDebugLine(Drawer.World, Segments[Idx].Key, Segments[Idx].Value, Color.ToFColor(true), false, -1.0f, bDepthTest ? 0 : 1, Thickness);
+		}
+	}
+
+	if (Drawer.AnimInstanceProxy)
+	{
+		for (int32 Idx = 0; Idx < Num; Idx++)
+		{
+			Drawer.AnimInstanceProxy->AnimDrawDebugLine(Segments[Idx].Key, Segments[Idx].Value, Color.ToFColor(true), false, -1.0f, Thickness, bDepthTest ? SDPG_World : SDPG_Foreground);
+		}
+	}
+#endif
+
+#if ENABLE_VISUAL_LOG
+	if (Drawer.VisualLoggerObject)
+	{
+		FVisualLogEntry* CurrentEntry = nullptr;
+		if (FVisualLogger::CheckVisualLogInputInternal(Drawer.VisualLoggerObject, Drawer.VisualLoggerCategory, Drawer.VisualLoggerVerbosity, /*OutWorld=*/ nullptr, &CurrentEntry))
+		{
+			const uint16 IntegerThickness = Thickness == 0.0f ? 0 : FMath::CeilToInt(FMath::Max(2.0f * Thickness, 0.0f));
+			FVisualLogShapeElement Element(TEXT(""), Color.ToFColor(false), IntegerThickness, Drawer.VisualLoggerCategory);
+			Element.Type = EVisualLoggerShapeElement::Segment;
+			Element.Verbosity = Drawer.VisualLoggerVerbosity;
+			Element.Points.Reserve(2 * Num);
+			for (int32 Idx = 0; Idx < Num; Idx++)
+			{
+				Element.Points.Add(Segments[Idx].Key);
+				Element.Points.Add(Segments[Idx].Value);
+			}
+			CurrentEntry->ElementsToDraw.Add(Element);
+		}
+	}
+#endif
+}
+
+void UDrawDebugLibrary::DrawDebugTriangularBasePyramid(const FDebugDrawer& Drawer, const FVector& Location, const FRotator& Rotation, const FLinearColor Color, const float Length, const float Width, const float Thickness, const bool bDepthTest)
+{
+	const FVector V0 = Location + Rotation.RotateVector(FVector(0, 0, Length/2));
+	const FVector V1 = Location + Rotation.RotateVector(FVector(-Width / 4 * FMath::Sqrt(3.0f), -Width / 2, -Length / 2));
+	const FVector V2 = Location + Rotation.RotateVector(FVector(-Width / 4 * FMath::Sqrt(3.0f), +Width / 2, -Length / 2));
+	const FVector V3 = Location + Rotation.RotateVector(FVector(+Width / 4 * FMath::Sqrt(3.0f), 0.0f, -Length / 2));
+
+	DrawDebugLinesPairsArrayView(Drawer, {
+		{ V0, V1 },
+		{ V0, V2 },
+		{ V0, V3 },
+		{ V1, V2 },
+		{ V2, V3 },
+		{ V3, V1 },
+	}, Color, Thickness, bDepthTest);
+}
+
+void UDrawDebugLibrary::DrawDebugSquareBasePyramid(const FDebugDrawer& Drawer, const FVector& Location, const FRotator& Rotation, const FLinearColor Color, const float Length, const float Width, const float Thickness, const bool bDepthTest)
+{
+	const FVector V0 = Location + Rotation.RotateVector(FVector(0, 0, Length / 2));
+	const FVector V1 = Location + Rotation.RotateVector(FVector(-Width / 2, -Width / 2, -Length / 2));
+	const FVector V2 = Location + Rotation.RotateVector(FVector(-Width / 2, +Width / 2, -Length / 2));
+	const FVector V3 = Location + Rotation.RotateVector(FVector(+Width / 2, +Width / 2, -Length / 2));
+	const FVector V4 = Location + Rotation.RotateVector(FVector(+Width / 2, -Width / 2, -Length / 2));
+
+	DrawDebugLinesPairsArrayView(Drawer, {
+		{ V0, V1 },
+		{ V0, V2 },
+		{ V0, V3 },
+		{ V0, V4 },
+		{ V1, V2 },
+		{ V2, V3 },
+		{ V3, V4 },
+		{ V4, V1 },
+		}, Color, Thickness, bDepthTest);
+}
+
+void UDrawDebugLibrary::DrawDebugCone(const FDebugDrawer& Drawer, const FVector& Location, const FRotator& Rotation, const FLinearColor Color, const float Length, const float Radius, const int32 Segments, const float Thickness, const bool bDepthTest)
+{
+	const int32 Num = FMath::Max(Segments, 2);
+
+	const FVector Peak = Location + Rotation.RotateVector(FVector(0, 0, Length / 2));
+
+	TArray<TPair<FVector, FVector>, TInlineAllocator<9 * 2>> LineSegments;
+	LineSegments.Reserve(Num * 2);
+
+	for (int32 Idx = 0; Idx < Num; Idx++)
+	{
+		const FVector V0 = FVector(Radius * FMath::Sin(UE_TWO_PI * (float)(Idx + 0) / Num), Radius * FMath::Cos(UE_TWO_PI * (float)(Idx + 0) / Num), -Length / 2);
+		const FVector V1 = FVector(Radius * FMath::Sin(UE_TWO_PI * (float)(Idx + 1) / Num), Radius * FMath::Cos(UE_TWO_PI * (float)(Idx + 1) / Num), -Length / 2);
+		LineSegments.Add({ Peak, Location + Rotation.RotateVector(V0) });
+		LineSegments.Add({ Location + Rotation.RotateVector(V0), Location + Rotation.RotateVector(V1) });
+	}
+
+	DrawDebugLinesPairsArrayView(Drawer, LineSegments, Color, Thickness, bDepthTest);
+}
+
+void UDrawDebugLibrary::DrawDebugArc(const FDebugDrawer& Drawer, const FVector& Location, const FRotator& Rotation, const float Angle, const FLinearColor Color, const float Radius, const float Thickness, const int32 Segments, const bool bDepthTest)
+{
+
+	TArray<TPair<FVector, FVector>, TInlineAllocator<17>> LineSegments;
+	LineSegments.Reserve(Segments);
+
+	for (int32 SegmentIdx = 0; SegmentIdx < Segments; SegmentIdx++)
+	{
+		const float Angle0 = (((float)SegmentIdx + 0) / Segments) * FMath::DegreesToRadians(Angle);
+		const float Angle1 = (((float)SegmentIdx + 1) / Segments) * FMath::DegreesToRadians(Angle);
+		const FVector Location0 = Location + Rotation.RotateVector(Radius * FVector(FMath::Cos(Angle0), FMath::Sin(Angle0), 0.0f));
+		const FVector Location1 = Location + Rotation.RotateVector(Radius * FVector(FMath::Cos(Angle1), FMath::Sin(Angle1), 0.0f));
+		LineSegments.Add({ Location0, Location1 });
+	}
+
+	DrawDebugLinesPairsArrayView(Drawer, LineSegments, Color, Thickness, bDepthTest);
+}
+
+void UDrawDebugLibrary::DrawDebugCircle(const FDebugDrawer& Drawer, const FVector& Location, const FRotator& Rotation, const FLinearColor Color, const float Radius, const float Thickness, const int32 Segments, const bool bDepthTest)
+{
+	DrawDebugArc(Drawer, Location, Rotation, 360.0f, Color, Radius, Thickness, Segments, bDepthTest);
+}
+
+void UDrawDebugLibrary::DrawDebugCircleOutline(const FDebugDrawer& Drawer, const FVector& Location, const FRotator& Rotation, const FLinearColor Color, const float InnerRadius, const float OuterRadius, const float Thickness, const int32 Segments, const bool bDepthTest)
+{
+	DrawDebugCircle(Drawer, Location, Rotation, Color, InnerRadius, Thickness, Segments, bDepthTest);
+	DrawDebugCircle(Drawer, Location, Rotation, Color, OuterRadius, Thickness, Segments, bDepthTest);
+}
+
+void UDrawDebugLibrary::DrawDebugLocator(const FDebugDrawer& Drawer, const FVector& Location, const FRotator& Rotation, const FLinearColor Color, const float DrawRadius, const float Thickness, const bool bDepthTest)
+{
+
+	DrawDebugLinesPairsArrayView(Drawer, {
+		{ Location - Rotation.RotateVector(FVector(DrawRadius, 0, 0)), Location + Rotation.RotateVector(FVector(DrawRadius, 0, 0)) },
+		{ Location - Rotation.RotateVector(FVector(0, DrawRadius, 0)), Location + Rotation.RotateVector(FVector(0, DrawRadius, 0)) },
+		{ Location - Rotation.RotateVector(FVector(0, 0, DrawRadius)), Location + Rotation.RotateVector(FVector(0, 0, DrawRadius)) },
+		}, Color, Thickness, bDepthTest);
+}
+
+void UDrawDebugLibrary::DrawDebugCrossLocator(const FDebugDrawer& Drawer, const FVector& Location, const FRotator& Rotation, const FLinearColor Color, const float DrawRadius, const float Thickness, const bool bDepthTest)
+{
+	DrawDebugLinesPairsArrayView(Drawer, {
+		{ Location + Rotation.RotateVector(FVector(+DrawRadius, +DrawRadius, -DrawRadius)), Location + Rotation.RotateVector(FVector(-DrawRadius, -DrawRadius, +DrawRadius)) },
+		{ Location + Rotation.RotateVector(FVector(-DrawRadius, +DrawRadius, -DrawRadius)), Location + Rotation.RotateVector(FVector(+DrawRadius, -DrawRadius, +DrawRadius)) },
+		{ Location + Rotation.RotateVector(FVector(+DrawRadius, -DrawRadius, -DrawRadius)), Location + Rotation.RotateVector(FVector(-DrawRadius, +DrawRadius, +DrawRadius)) },
+		{ Location + Rotation.RotateVector(FVector(-DrawRadius, -DrawRadius, -DrawRadius)), Location + Rotation.RotateVector(FVector(+DrawRadius, +DrawRadius, +DrawRadius)) },
+		}, Color, Thickness, bDepthTest);
+}
+
+void UDrawDebugLibrary::DrawDebugBox(const FDebugDrawer& Drawer, const FVector& Location, const FRotator& Rotation, const FVector HalfExtents, const FLinearColor Color, const float Thickness, const bool bDepthTest)
+{
+	const FVector C0 = Location + Rotation.RotateVector(FVector(+1, +1, -1) * HalfExtents);
+	const FVector C1 = Location + Rotation.RotateVector(FVector(+1, -1, -1) * HalfExtents);
+	const FVector C2 = Location + Rotation.RotateVector(FVector(-1, -1, -1) * HalfExtents);
+	const FVector C3 = Location + Rotation.RotateVector(FVector(-1, +1, -1) * HalfExtents);
+	const FVector C4 = Location + Rotation.RotateVector(FVector(+1, +1, +1) * HalfExtents);
+	const FVector C5 = Location + Rotation.RotateVector(FVector(+1, -1, +1) * HalfExtents);
+	const FVector C6 = Location + Rotation.RotateVector(FVector(-1, -1, +1) * HalfExtents);
+	const FVector C7 = Location + Rotation.RotateVector(FVector(-1, +1, +1) * HalfExtents);
+
+	DrawDebugLinesPairsArrayView(Drawer, {
+		{ C0, C1 },
+		{ C1, C2 },
+		{ C2, C3 },
+		{ C3, C0 },
+
+		{ C4, C5 },
+		{ C5, C6 },
+		{ C6, C7 },
+		{ C7, C4 },
+
+		{ C0, C4 },
+		{ C1, C5 },
+		{ C2, C6 },
+		{ C3, C7 },
+		}, Color, Thickness, bDepthTest);
+}
+
+void UDrawDebugLibrary::DrawDebugSphere(
+	const FDebugDrawer& Drawer,
+	const FVector& Location,
+	const FRotator& Rotation,
+	const FLinearColor Color,
+	const float Radius,
+	const float Thickness,
+	const int32 Segments,
+	const bool bDepthTest)
+{
+	const int32 SegmentNum = (FMath::Max(Segments, 4) / 2) * 2;
+
+	TArray<TPair<FVector, FVector>, TInlineAllocator<8 * 8>> LineSegments;
+	LineSegments.Reserve(SegmentNum * SegmentNum);
+
+	for (int32 X = 0; X < SegmentNum; X++)
+	{
+		const float XA0 = ((float)(X + 0) / SegmentNum) * UE_TWO_PI;
+		const float XA1 = ((float)(X + 1) / SegmentNum) * UE_TWO_PI;
+
+		for (int32 Y = 0; Y < SegmentNum / 2; Y++)
+		{
+			const float YA0 = ((float)(Y + 0) / SegmentNum) * UE_TWO_PI;
+			const float YA1 = ((float)(Y + 1) / SegmentNum) * UE_TWO_PI;
+
+			const FVector V0 = Location + Rotation.RotateVector(FVector(FMath::Cos(XA0) * FMath::Sin(YA0), FMath::Sin(XA0) * FMath::Sin(YA0), FMath::Cos(YA0)) * Radius);
+			const FVector V1 = Location + Rotation.RotateVector(FVector(FMath::Cos(XA1) * FMath::Sin(YA0), FMath::Sin(XA1) * FMath::Sin(YA0), FMath::Cos(YA0)) * Radius);
+			const FVector V2 = Location + Rotation.RotateVector(FVector(FMath::Cos(XA0) * FMath::Sin(YA1), FMath::Sin(XA0) * FMath::Sin(YA1), FMath::Cos(YA1)) * Radius);
+
+			LineSegments.Add({ V0, V1 });
+			LineSegments.Add({ V0, V2 });
+		}
+	}
+
+	DrawDebugLinesPairsArrayView(Drawer, LineSegments, Color, Thickness, bDepthTest);
+
+}
+
+void UDrawDebugLibrary::DrawDebugCapsule(const FDebugDrawer& Drawer, const FVector& Location, const FRotator& Rotation, const FLinearColor Color, const float Radius, const float HalfLength, const float Thickness, const int32 Segments, const bool bDepthTest)
+{
+	const int32 SegmentNum = (FMath::Max(Segments, 4) / 4) * 4;
+
+	TArray<TPair<FVector, FVector>, TInlineAllocator<8 * (8 + 3)>> LineSegments;
+	LineSegments.Reserve(SegmentNum * (SegmentNum + 3));
+
+	for (int32 X = 0; X < SegmentNum; X++)
+	{
+		const float XA0 = ((float)(X + 0) / SegmentNum) * UE_TWO_PI;
+		const float XA1 = ((float)(X + 1) / SegmentNum) * UE_TWO_PI;
+
+		for (int32 Y = 0; Y < SegmentNum / 4; Y++)
+		{
+			const float YA0 = ((float)(Y + 0) / SegmentNum) * UE_TWO_PI;
+			const float YA1 = ((float)(Y + 1) / SegmentNum) * UE_TWO_PI;
+
+			const FVector V0T = Rotation.RotateVector(FVector(0.0f, 0.0f, +HalfLength) + FVector(FMath::Cos(XA0) * FMath::Sin(YA0), FMath::Sin(XA0) * FMath::Sin(YA0), +FMath::Cos(YA0)) * Radius) + Location;
+			const FVector V1T = Rotation.RotateVector(FVector(0.0f, 0.0f, +HalfLength) + FVector(FMath::Cos(XA1) * FMath::Sin(YA0), FMath::Sin(XA1) * FMath::Sin(YA0), +FMath::Cos(YA0)) * Radius) + Location;
+			const FVector V2T = Rotation.RotateVector(FVector(0.0f, 0.0f, +HalfLength) + FVector(FMath::Cos(XA0) * FMath::Sin(YA1), FMath::Sin(XA0) * FMath::Sin(YA1), +FMath::Cos(YA1)) * Radius) + Location;
+
+			const FVector V0B = Rotation.RotateVector(FVector(0.0f, 0.0f, -HalfLength) + FVector(FMath::Cos(XA0) * FMath::Sin(YA0), FMath::Sin(XA0) * FMath::Sin(YA0), -FMath::Cos(YA0)) * Radius) + Location;
+			const FVector V1B = Rotation.RotateVector(FVector(0.0f, 0.0f, -HalfLength) + FVector(FMath::Cos(XA1) * FMath::Sin(YA0), FMath::Sin(XA1) * FMath::Sin(YA0), -FMath::Cos(YA0)) * Radius) + Location;
+			const FVector V2B = Rotation.RotateVector(FVector(0.0f, 0.0f, -HalfLength) + FVector(FMath::Cos(XA0) * FMath::Sin(YA1), FMath::Sin(XA0) * FMath::Sin(YA1), -FMath::Cos(YA1)) * Radius) + Location;
+
+			LineSegments.Add({ V0T, V1T });
+			LineSegments.Add({ V0T, V2T });
+			LineSegments.Add({ V0B, V1B });
+			LineSegments.Add({ V0B, V2B });
+		}
+
+		const float YA0 = ((float)(SegmentNum / 4) / SegmentNum) * UE_TWO_PI;
+
+		const FVector V0T = Rotation.RotateVector(FVector(0.0f, 0.0f, +HalfLength) + FVector(FMath::Cos(XA0) * FMath::Sin(YA0), FMath::Sin(XA0) * FMath::Sin(YA0), +FMath::Cos(YA0)) * Radius) + Location;
+		const FVector V1T = Rotation.RotateVector(FVector(0.0f, 0.0f, +HalfLength) + FVector(FMath::Cos(XA1) * FMath::Sin(YA0), FMath::Sin(XA1) * FMath::Sin(YA0), +FMath::Cos(YA0)) * Radius) + Location;
+
+		const FVector V0B = Rotation.RotateVector(FVector(0.0f, 0.0f, -HalfLength) + FVector(FMath::Cos(XA0) * FMath::Sin(YA0), FMath::Sin(XA0) * FMath::Sin(YA0), -FMath::Cos(YA0)) * Radius) + Location;
+		const FVector V1B = Rotation.RotateVector(FVector(0.0f, 0.0f, -HalfLength) + FVector(FMath::Cos(XA1) * FMath::Sin(YA0), FMath::Sin(XA1) * FMath::Sin(YA0), -FMath::Cos(YA0)) * Radius) + Location;
+
+		LineSegments.Add({ V0T, V1T });
+		LineSegments.Add({ V0B, V1B });
+
+		const FVector E0 = Rotation.RotateVector(FVector(0.0f, 0.0f, +HalfLength) + FVector(FMath::Cos(XA0), FMath::Sin(XA0), 0.0f) * Radius) + Location;
+		const FVector E1 = Rotation.RotateVector(FVector(0.0f, 0.0f, -HalfLength) + FVector(FMath::Cos(XA0), FMath::Sin(XA0), 0.0f) * Radius) + Location;
+		LineSegments.Add({ E0, E1 });
+	}
+
+	DrawDebugLinesPairsArrayView(Drawer, LineSegments, Color, Thickness, bDepthTest);
+}
+
+
+void UDrawDebugLibrary::DrawDebugFrustum(const FDebugDrawer& Drawer, const FMatrix& FrustumToWorld, const FLinearColor Color, const float Thickness, const bool bDepthTest)
+{
+	FVector Vertices[2][2][2];
+	for (int32 Z = 0; Z < 2; Z++)
+	{
+		for (int32 Y = 0; Y < 2; Y++)
+		{
+			for (int32 X = 0; X < 2; X++)
+			{
+				const FVector4 UnprojectedVertex = FrustumToWorld.TransformFVector4(
+					FVector4(
+						(X ? -1.0f : 1.0f),
+						(Y ? -1.0f : 1.0f),
+						(Z ? 0.0f : 1.0f),
+						1.0f));
+				Vertices[X][Y][Z] = FVector(UnprojectedVertex) / UnprojectedVertex.W;
+			}
+		}
+	}
+
+	DrawDebugLinesPairsArrayView(Drawer, {
+		{ Vertices[0][0][0], Vertices[0][0][1] },
+		{ Vertices[1][0][0], Vertices[1][0][1] },
+		{ Vertices[0][1][0], Vertices[0][1][1] },
+		{ Vertices[1][1][0], Vertices[1][1][1] },
+
+		{ Vertices[0][0][0], Vertices[0][1][0] },
+		{ Vertices[1][0][0], Vertices[1][1][0] },
+		{ Vertices[0][0][1], Vertices[0][1][1] },
+		{ Vertices[1][0][1], Vertices[1][1][1] },
+
+		{ Vertices[0][0][0], Vertices[1][0][0] },
+		{ Vertices[0][1][0], Vertices[1][1][0] },
+		{ Vertices[0][0][1], Vertices[1][0][1] },
+		{ Vertices[0][1][1], Vertices[1][1][1] },
+		}, Color, Thickness, bDepthTest);
+}
+
+void UDrawDebugLibrary::DrawDebugArrow(
+	const FDebugDrawer& Drawer,
+	const FVector& StartLocation,
+	const FVector& EndLocation,
+	const FLinearColor Color,
+	const float ArrowHeadSize,
+	const float Thickness,
+	const bool bDepthTest)
+{
+	const FQuat Rotation = UE::DrawDebugLibrary::Private::FindOrientation(StartLocation, EndLocation);
+
+	const float ArrowSqrt = FMath::Sqrt(ArrowHeadSize);
+
+	DrawDebugLinesPairsArrayView(Drawer, {
+		{ StartLocation, EndLocation },
+		{ EndLocation, EndLocation + Rotation.RotateVector(FVector(-ArrowSqrt, +ArrowSqrt, 0)) },
+		{ EndLocation, EndLocation + Rotation.RotateVector(FVector(-ArrowSqrt, -ArrowSqrt, 0)) },
+		}, Color, Thickness, bDepthTest);
+}
+
+void UDrawDebugLibrary::DrawDebugOrientedArrow(const FDebugDrawer& Drawer, const FVector& Location, const FRotator& Rotation, const float Length, const FLinearColor Color, const float ArrowHeadSize, const float Thickness, const bool bDepthTest)
+{
+	const float ArrowSqrt = FMath::Sqrt(ArrowHeadSize);
+
+	DrawDebugLinesPairsArrayView(Drawer, {
+		{ Location, Location + Rotation.RotateVector(FVector(Length, 0.0f, 0.0f)) },
+		{ Location + Rotation.RotateVector(FVector(Length, 0.0f, 0.0f)), Location + Rotation.RotateVector(FVector(Length - ArrowSqrt, +ArrowSqrt, 0.0f)) },
+		{ Location + Rotation.RotateVector(FVector(Length, 0.0f, 0.0f)), Location + Rotation.RotateVector(FVector(Length - ArrowSqrt, -ArrowSqrt, 0.0f)) },
+		}, Color, Thickness, bDepthTest);
+}
+
+void UDrawDebugLibrary::DrawDebugFourWayArrow(const FDebugDrawer& Drawer, const FVector& StartLocation, const FVector& EndLocation, const FLinearColor Color, const float ArrowHeadSize, const float Thickness, const bool bDepthTest)
+{
+	const FQuat Rotation = UE::DrawDebugLibrary::Private::FindOrientation(StartLocation, EndLocation);
+
+	const float ArrowSqrt = FMath::Sqrt(ArrowHeadSize);
+
+	DrawDebugLinesPairsArrayView(Drawer, {
+		{ StartLocation, EndLocation },
+		{ EndLocation, EndLocation + Rotation.RotateVector(FVector(-ArrowSqrt, +ArrowSqrt, 0)) },
+		{ EndLocation, EndLocation + Rotation.RotateVector(FVector(-ArrowSqrt, -ArrowSqrt, 0)) },
+		{ EndLocation, EndLocation + Rotation.RotateVector(FVector(-ArrowSqrt, 0, +ArrowSqrt)) },
+		{ EndLocation, EndLocation + Rotation.RotateVector(FVector(-ArrowSqrt, 0, -ArrowSqrt)) },
+		}, Color, Thickness, bDepthTest);
+}
+
+void UDrawDebugLibrary::DrawDebugGroundTargetArrow(const FDebugDrawer& Drawer, const FVector& Location, const FLinearColor Color, const float Length, const float ArrowHeadSize, const float Radius, const int32 Segements, const float Thickness, const bool bDepthTest)
+{
+	DrawDebugFourWayArrow(
+		Drawer,
+		Location + FVector(0.0f, 0.0f, Length),
+		Location,
+		Color,
+		ArrowHeadSize,
+		Thickness,
+		bDepthTest);
+
+	DrawDebugCircle(Drawer, Location, FRotator::ZeroRotator, Color, Radius, Thickness, Segements, bDepthTest);
+}
+
+void UDrawDebugLibrary::DrawTriangularArrow(const FDebugDrawer& Drawer, const FVector& StartLocation, const FVector& EndLocation, const FLinearColor Color, const float ArrowHeadSize, const float Thickness, const bool bDepthTest)
+{
+	const FQuat Rotation = UE::DrawDebugLibrary::Private::FindOrientation(StartLocation, EndLocation);
+
+	const float ArrowSqrt = FMath::Sqrt(ArrowHeadSize);
+
+	DrawDebugLine(Drawer, StartLocation, EndLocation, Color, Thickness, bDepthTest);
+	DrawDebugTriangularBasePyramid(Drawer, EndLocation + Rotation.RotateVector(FVector(-ArrowSqrt / 2, 0, 0)), (Rotation * FQuat::MakeFromRotationVector(FVector(0, UE_PI / 2, 0))).Rotator(), Color, ArrowSqrt, ArrowSqrt, Thickness, bDepthTest);
+}
+
+void UDrawDebugLibrary::DrawSquareArrow(const FDebugDrawer& Drawer, const FVector& StartLocation, const FVector& EndLocation, const FLinearColor Color, const float ArrowHeadSize, const float Thickness, const bool bDepthTest)
+{
+	const FQuat Rotation = UE::DrawDebugLibrary::Private::FindOrientation(StartLocation, EndLocation);
+
+	const float ArrowSqrt = FMath::Sqrt(ArrowHeadSize);
+
+	DrawDebugLine(Drawer, StartLocation, EndLocation, Color, Thickness, bDepthTest);
+	DrawDebugSquareBasePyramid(Drawer, EndLocation + Rotation.RotateVector(FVector(-ArrowSqrt / 2, 0, 0)), (Rotation * FQuat::MakeFromRotationVector(FVector(0, UE_PI / 2, 0))).Rotator(), Color, ArrowSqrt, ArrowSqrt, Thickness, bDepthTest);
+}
+
+void UDrawDebugLibrary::DrawConeArrow(const FDebugDrawer& Drawer, const FVector& StartLocation, const FVector& EndLocation, const FLinearColor Color, const float ArrowHeadSize, const float Thickness, const int32 Segments, const bool bDepthTest)
+{
+	const FQuat Rotation = UE::DrawDebugLibrary::Private::FindOrientation(StartLocation, EndLocation);
+
+	const float ArrowSqrt = FMath::Sqrt(ArrowHeadSize);
+
+	DrawDebugLine(Drawer, StartLocation, EndLocation, Color, Thickness, bDepthTest);
+	DrawDebugCone(Drawer, EndLocation + Rotation.RotateVector(FVector(-ArrowSqrt / 2, 0, 0)), (Rotation * FQuat::MakeFromRotationVector(FVector(0, UE_PI / 2, 0))).Rotator(), Color, ArrowSqrt, ArrowSqrt / 2, Segments, Thickness, bDepthTest);
+}
+
+void UDrawDebugLibrary::DrawDebugFlatArrow(
+	const FDebugDrawer& Drawer,
+	const FVector& Location,
+	const FRotator& Rotation,
+	const FLinearColor Color,
+	const float Length,
+	const float Width,
+	const float Thickness,
+	const bool bDepthTest)
+{
+	FVector ArrowPoints[7];
+	ArrowPoints[0] = Location - Rotation.RotateVector(FVector(0, Width / 2.0f, 0));
+	ArrowPoints[1] = Location + Rotation.RotateVector(FVector(0, Width / 2.0f, 0));
+	ArrowPoints[2] = ArrowPoints[0] + Rotation.RotateVector(FVector(Length / 3.0f, 0, 0));
+	ArrowPoints[3] = ArrowPoints[1] + Rotation.RotateVector(FVector(Length / 3.0f, 0, 0));
+	ArrowPoints[4] = ArrowPoints[2] - Rotation.RotateVector(FVector(0, Width / 2.0f, 0));
+	ArrowPoints[5] = ArrowPoints[3] + Rotation.RotateVector(FVector(0, Width / 2.0f, 0));
+	ArrowPoints[6] = Location + Rotation.RotateVector(FVector(Length, 0, 0));
+
+	DrawDebugLinesPairsArrayView(Drawer, {
+		{ ArrowPoints[0], ArrowPoints[1] },
+		{ ArrowPoints[0], ArrowPoints[2] },
+		{ ArrowPoints[1], ArrowPoints[3] },
+		{ ArrowPoints[2], ArrowPoints[4] },
+		{ ArrowPoints[3], ArrowPoints[5] },
+		{ ArrowPoints[4], ArrowPoints[6] },
+		{ ArrowPoints[5], ArrowPoints[6] },
+		}, Color, Thickness, bDepthTest);
+}
+
+void UDrawDebugLibrary::DrawDebugCatmullRomSplineStart(const FDebugDrawer& Drawer, const FVector& V0, const FVector& V1, const FVector& V2, const bool bMonotonic, const FLinearColor Color, const float Thickness, const int32 Segments, const bool bDepthTest)
+{
+	const int32 SegmentNum = FMath::Max(Segments, 2);
+
+	TArray<TPair<FVector, FVector>, TInlineAllocator<14>> LineSegments;
+	LineSegments.Reserve(SegmentNum - 1);
+
+	if (bMonotonic)
+	{
+		for (int32 Idx = 0; Idx < Segments - 1; Idx++)
+		{
+			const float Alpha0 = ((float)(Idx + 0)) / (Segments - 1);
+			const float Alpha1 = ((float)(Idx + 1)) / (Segments - 1);
+			const FVector Location0 = UE::DrawDebugLibrary::Private::InterpolateCubicMonoStart(V0, V1, V2, Alpha0);
+			const FVector Location1 = UE::DrawDebugLibrary::Private::InterpolateCubicMonoStart(V0, V1, V2, Alpha1);
+			LineSegments.Add({ Location0, Location1 });
+		}
+	}
+	else
+	{
+		for (int32 Idx = 0; Idx < Segments - 1; Idx++)
+		{
+			const float Alpha0 = ((float)(Idx + 0)) / (Segments - 1);
+			const float Alpha1 = ((float)(Idx + 1)) / (Segments - 1);
+			const FVector Location0 = UE::DrawDebugLibrary::Private::InterpolateCubicStart(V0, V1, V2, Alpha0);
+			const FVector Location1 = UE::DrawDebugLibrary::Private::InterpolateCubicStart(V0, V1, V2, Alpha1);
+			LineSegments.Add({ Location0, Location1 });
+		}
+	}
+
+	DrawDebugLinesPairsArrayView(Drawer, LineSegments, Color, Thickness, bDepthTest);
+}
+
+void UDrawDebugLibrary::DrawDebugCatmullRomSplineEnd(const FDebugDrawer& Drawer, const FVector& V0, const FVector& V1, const FVector& V2, const bool bMonotonic, const FLinearColor Color, const float Thickness, const int32 Segments, const bool bDepthTest)
+{
+	const int32 SegmentNum = FMath::Max(Segments, 2);
+
+	TArray<TPair<FVector, FVector>, TInlineAllocator<14>> LineSegments;
+	LineSegments.Reserve(SegmentNum - 1);
+
+	if (bMonotonic)
+	{
+		for (int32 Idx = 0; Idx < Segments - 1; Idx++)
+		{
+			const float Alpha0 = ((float)(Idx + 0)) / (Segments - 1);
+			const float Alpha1 = ((float)(Idx + 1)) / (Segments - 1);
+			const FVector Location0 = UE::DrawDebugLibrary::Private::InterpolateCubicMonoEnd(V0, V1, V2, Alpha0);
+			const FVector Location1 = UE::DrawDebugLibrary::Private::InterpolateCubicMonoEnd(V0, V1, V2, Alpha1);
+			LineSegments.Add({ Location0, Location1 });
+		}
+	}
+	else
+	{
+		for (int32 Idx = 0; Idx < Segments - 1; Idx++)
+		{
+			const float Alpha0 = ((float)(Idx + 0)) / (Segments - 1);
+			const float Alpha1 = ((float)(Idx + 1)) / (Segments - 1);
+			const FVector Location0 = UE::DrawDebugLibrary::Private::InterpolateCubicEnd(V0, V1, V2, Alpha0);
+			const FVector Location1 = UE::DrawDebugLibrary::Private::InterpolateCubicEnd(V0, V1, V2, Alpha1);
+			LineSegments.Add({ Location0, Location1 });
+		}
+	}
+
+	DrawDebugLinesPairsArrayView(Drawer, LineSegments, Color, Thickness, bDepthTest);
+}
+
+void UDrawDebugLibrary::DrawDebugCatmullRomSplineSection(const FDebugDrawer& Drawer, const FVector& V0, const FVector& V1, const FVector& V2, const FVector& V3, const bool bMonotonic, const FLinearColor Color, const float Thickness, const int32 Segments, const bool bDepthTest)
+{
+	const int32 SegmentNum = FMath::Max(Segments, 2);
+
+	TArray<TPair<FVector, FVector>, TInlineAllocator<14>> LineSegments;
+	LineSegments.Reserve(SegmentNum - 1);
+
+	if (bMonotonic)
+	{
+		for (int32 Idx = 0; Idx < Segments - 1; Idx++)
+		{
+			const float Alpha0 = ((float)(Idx + 0)) / (Segments - 1);
+			const float Alpha1 = ((float)(Idx + 1)) / (Segments - 1);
+			const FVector Location0 = UE::DrawDebugLibrary::Private::InterpolateCubicMono(V0, V1, V2, V3, Alpha0);
+			const FVector Location1 = UE::DrawDebugLibrary::Private::InterpolateCubicMono(V0, V1, V2, V3, Alpha1);
+			LineSegments.Add({ Location0, Location1 });
+		}
+	}
+	else
+	{
+		for (int32 Idx = 0; Idx < Segments - 1; Idx++)
+		{
+			const float Alpha0 = ((float)(Idx + 0)) / (Segments - 1);
+			const float Alpha1 = ((float)(Idx + 1)) / (Segments - 1);
+			const FVector Location0 = UE::DrawDebugLibrary::Private::InterpolateCubic(V0, V1, V2, V3, Alpha0);
+			const FVector Location1 = UE::DrawDebugLibrary::Private::InterpolateCubic(V0, V1, V2, V3, Alpha1);
+			LineSegments.Add({ Location0, Location1 });
+		}
+	}
+
+	DrawDebugLinesPairsArrayView(Drawer, LineSegments, Color, Thickness, bDepthTest);
+}
+
+
+void UDrawDebugLibrary::DrawDebugCatmullRomSpline(const FDebugDrawer& Drawer, const TArray<FVector>& Points, const bool bMonotonic, const FLinearColor Color, const float Thickness, const int32 Segments, const bool bDepthTest)
+{
+	DrawDebugCatmullRomSplineArrayView(Drawer, Points, bMonotonic, Color, Thickness, Segments, bDepthTest);
+}
+
+void UDrawDebugLibrary::DrawDebugCatmullRomSplineArrayView(const FDebugDrawer& Drawer, const TArrayView<const FVector> Points, const bool bMonotonic, const FLinearColor Color, const float Thickness, const int32 Segments, const bool bDepthTest)
+{
+	const int32 Num = Points.Num();
+	
+	if (Num == 0) { return; }
+	if (Num == 1)
+	{
+		DrawDebugPoint(Drawer, Points[0], Color, Thickness, bDepthTest);
+	}
+
+	if (Num == 2)
+	{
+		DrawDebugLine(Drawer, Points[0], Points[1], Color, Thickness, bDepthTest);
+	}
+
+	DrawDebugCatmullRomSplineStart(Drawer, Points[0], Points[1], Points[2], bMonotonic, Color, Thickness, Segments, bDepthTest);
+
+	for (int32 Idx = 0; Idx < Num - 3; Idx++)
+	{
+		DrawDebugCatmullRomSplineSection(Drawer, Points[Idx + 0], Points[Idx + 1], Points[Idx + 2], Points[Idx + 3], bMonotonic, Color, Thickness, Segments, bDepthTest);
+	}
+
+	DrawDebugCatmullRomSplineEnd(Drawer, Points[Num - 3], Points[Num - 2], Points[Num - 1], bMonotonic, Color, Thickness, Segments, bDepthTest);
+}
+
+
+
+void UDrawDebugLibrary::DrawDebugAngle(const FDebugDrawer& Drawer, const FVector& Location, const FRotator& Rotation, const float Angle, const float LineLength, const float AngleRadius, const FLinearColor Color, const float Thickness, const int32 Segments, const bool bDepthTest)
+{
+	DrawDebugArc(Drawer, Location, Rotation, Angle, Color, AngleRadius, Thickness, Segments, bDepthTest);
+	DrawDebugLine(Drawer, Location, Location + Rotation.RotateVector(FVector(LineLength, 0, 0)), Color, Thickness, bDepthTest);
+	DrawDebugLine(Drawer, Location, Location + Rotation.RotateVector(LineLength * FVector(FMath::Cos(FMath::DegreesToRadians(Angle)), FMath::Sin(FMath::DegreesToRadians(Angle)), 0)), Color, Thickness, bDepthTest);
+}
+
+void UDrawDebugLibrary::DrawDebugLocation(const FDebugDrawer& Drawer, const FVector& Location, const FLinearColor Color, const float DrawRadius, const float Thickness, const int32 Segments, const bool bDepthTest)
+{
+	DrawDebugSphere(Drawer, Location, FRotator::ZeroRotator, Color, DrawRadius, Thickness, Segments, bDepthTest);
+}
+
+void UDrawDebugLibrary::DrawDebugLocations(const FDebugDrawer& Drawer, const TArray<FVector>& Locations, const FLinearColor Color, const float DrawRadius, const float Thickness, const int32 Segments, const bool bDepthTest)
+{
+	DrawDebugLocationsArrayView(Drawer, Locations, Color, DrawRadius, Thickness, Segments, bDepthTest);
+}
+
+void UDrawDebugLibrary::DrawDebugLocationsArrayView(const FDebugDrawer& Drawer, const TArrayView<const FVector> Locations, const FLinearColor Color, const float DrawRadius, const float Thickness, const int32 Segments, const bool bDepthTest)
+{
+	for (const FVector& Location : Locations)
+	{
+		DrawDebugLocation(Drawer, Location, Color, DrawRadius, Thickness, Segments, bDepthTest);
+	}
+}
+
+void UDrawDebugLibrary::DrawDebugRotation(
+	const FDebugDrawer& Drawer,
+	const FVector& Location,
+	const FRotator& Rotation,
+	const float DrawRadius,
+	const float Thickness,
+	const float Opacity,
+	const bool bDepthTest)
+{
+	DrawDebugLine(Drawer, Location, Location + Rotation.RotateVector(FVector(DrawRadius, 0.0f, 0.0f)), FLinearColor(1.0f, 0.0f, 0.0f, Opacity), Thickness, bDepthTest);
+	DrawDebugLine(Drawer, Location, Location + Rotation.RotateVector(FVector(0.0f, DrawRadius, 0.0f)), FLinearColor(0.0f, 1.0f, 0.0f, Opacity), Thickness, bDepthTest);
+	DrawDebugLine(Drawer, Location, Location + Rotation.RotateVector(FVector(0.0f, 0.0f, DrawRadius)), FLinearColor(0.0f, 0.0f, 1.0f, Opacity), Thickness, bDepthTest);
+}
+
+void UDrawDebugLibrary::DrawDebugDirection(const FDebugDrawer& Drawer, const FVector& Location, const FVector& Direction, const FLinearColor Color, const float DrawArrowLength, const float ArrowHeadScale, const float Thickness, const bool bDepthTest)
+{
+	DrawDebugArrow(Drawer, Location, Location + DrawArrowLength * Direction, Color, 0.2f * ArrowHeadScale * DrawArrowLength, Thickness, bDepthTest);
+}
+
+void UDrawDebugLibrary::DrawDebugVelocity(const FDebugDrawer& Drawer, const FVector& Location, const FVector& Velocity, const FLinearColor Color, const float DrawVelocityLineScale, const float Thickness, const bool bDepthTest)
+{
+	DrawDebugLine(Drawer, Location, Location + DrawVelocityLineScale * Velocity, Color, Thickness, bDepthTest);
+}
+
+void UDrawDebugLibrary::DrawDebugTransform(const FDebugDrawer& Drawer, const FTransform& Transform, const float DrawRadius, const float Thickness, const float Opacity, const bool bDepthTest)
+{
+	const FVector Location = Transform.GetLocation();
+	const FQuat Rotation = Transform.GetRotation();
+	const FVector Scale = Transform.GetScale3D();
+
+	DrawDebugLine(Drawer, Location, Location + Rotation.RotateVector(Scale * FVector(DrawRadius, 0, 0)), FLinearColor(1.0f, 0.0f, 0.0f, Opacity), Thickness, bDepthTest);
+	DrawDebugLine(Drawer, Location, Location + Rotation.RotateVector(Scale * FVector(0, DrawRadius, 0)), FLinearColor(0.0f, 1.0f, 0.0f, Opacity), Thickness, bDepthTest);
+	DrawDebugLine(Drawer, Location, Location + Rotation.RotateVector(Scale * FVector(0, 0, DrawRadius)), FLinearColor(0.0f, 0.0f, 1.0f, Opacity), Thickness, bDepthTest);
+}
+
+void UDrawDebugLibrary::DrawDebugEvent(const FDebugDrawer& Drawer, bool bTimeUntilEventKnown, const float TimeUntilEvent, const FVector& Location, const FLinearColor Color, const float Size, const float Thickness, const bool bDepthTest)
+{
+	DrawDebugLine(
+		Drawer,
+		Location + FVector(0, 0, +Size),
+		Location + FVector(0, 0, -Size),
+		Color, Thickness, bDepthTest);
+
+	DrawDebugSphere(Drawer, Location, FRotator::ZeroRotator, Color, Size * 0.1f, Thickness, 8, bDepthTest);
+
+	// TODO: Don't hard-code apprehension
+	if (bTimeUntilEventKnown && TimeUntilEvent < 1.0f && TimeUntilEvent > -1.0f)
+	{
+		DrawDebugSphere(Drawer, Location + FVector(0, 0, TimeUntilEvent * Size), FRotator::ZeroRotator, Color, Size * 0.25f, Thickness, 8, bDepthTest);
+	}
+}
+
+
+
+
+void UDrawDebugLibrary::DrawDebugChair(const FDebugDrawer& Drawer, const FVector& Location, const FRotator& Rotation, const float Width, const float SeatHeight, const float BackHeight, const float BackTilt, const FLinearColor Color, const float Thickness, const bool bDepthTest)
+{
+	const float HalfWidth = Width / 2.0f;
+	const FVector F0 = Location + Rotation.RotateVector(FVector(+HalfWidth, +HalfWidth, 0));
+	const FVector F1 = Location + Rotation.RotateVector(FVector(+HalfWidth, -HalfWidth, 0));
+	const FVector F2 = Location + Rotation.RotateVector(FVector(-HalfWidth, -HalfWidth, 0));
+	const FVector F3 = Location + Rotation.RotateVector(FVector(-HalfWidth, +HalfWidth, 0));
+
+	const FVector S0 = Location + Rotation.RotateVector(FVector(+HalfWidth, +HalfWidth, SeatHeight));
+	const FVector S1 = Location + Rotation.RotateVector(FVector(+HalfWidth, -HalfWidth, SeatHeight));
+	const FVector S2 = Location + Rotation.RotateVector(FVector(-HalfWidth, -HalfWidth, SeatHeight));
+	const FVector S3 = Location + Rotation.RotateVector(FVector(-HalfWidth, +HalfWidth, SeatHeight));
+
+	const FVector B2 = Location + Rotation.RotateVector(FVector(-HalfWidth - BackTilt, -HalfWidth, SeatHeight + BackHeight));
+	const FVector B3 = Location + Rotation.RotateVector(FVector(-HalfWidth - BackTilt, +HalfWidth, SeatHeight + BackHeight));
+
+	DrawDebugLinesPairsArrayView(Drawer, {
+		{ F0, S0 },
+		{ F1, S1 },
+		{ F2, S2 },
+		{ F3, S3 },
+
+		{ S0, S1 },
+		{ S1, S2 },
+		{ S2, S3 },
+		{ S3, S0 },
+
+		{ S2, B2 },
+		{ S3, B3 },
+		{ B2, B3 },
+		}, Color, Thickness, bDepthTest);
+}
+
+void UDrawDebugLibrary::DrawDebugDoor(const FDebugDrawer& Drawer, const FVector& Location, const FRotator& Rotation, const float OpenAngle, const float Width, const float Height, const float Inset, const float HandleHeight, const float HandleOffset, const float HandleRadius, const float HandleExtension, const bool bOuterHandle, const bool bInnerHandle, const bool bHandleOnLeft, const bool bDrawEntryArrow, const FLinearColor Color, const float Thickness, const bool bDepthTest)
+{
+	const float HalfWidth = Width / 2.0f;
+	const FVector F0 = Location + Rotation.RotateVector(FVector(0.0f, -HalfWidth, 0.0f));
+	const FVector F1 = Location + Rotation.RotateVector(FVector(0.0f, +HalfWidth, 0.0f));
+	const FVector F2 = Location + Rotation.RotateVector(FVector(0.0f, +HalfWidth, Height));
+	const FVector F3 = Location + Rotation.RotateVector(FVector(0.0f, -HalfWidth, Height));
+
+	const FVector LI0 = FVector(0.0f, -HalfWidth + Inset, Inset);
+	const FVector LI1 = FVector(0.0f, +HalfWidth - Inset, Inset);
+	const FVector LI2 = FVector(0.0f, +HalfWidth - Inset, Height - Inset);
+	const FVector LI3 = FVector(0.0f, -HalfWidth + Inset, Height - Inset);
+
+	const FVector OpenOffset = bHandleOnLeft ?
+		FVector(0.0f, +HalfWidth - Inset, 0.0f) :
+		FVector(0.0f, -HalfWidth + Inset, 0.0f);
+	const FQuat OpenRotation = FQuat::MakeFromRotationVector(FVector(0.0f, 0.0f, FMath::DegreesToRadians(-OpenAngle)));
+
+	const FVector I0 = Location + Rotation.RotateVector(OpenRotation.RotateVector(LI0 - OpenOffset) + OpenOffset);
+	const FVector I1 = Location + Rotation.RotateVector(OpenRotation.RotateVector(LI1 - OpenOffset) + OpenOffset);
+	const FVector I2 = Location + Rotation.RotateVector(OpenRotation.RotateVector(LI2 - OpenOffset) + OpenOffset);
+	const FVector I3 = Location + Rotation.RotateVector(OpenRotation.RotateVector(LI3 - OpenOffset) + OpenOffset);
+
+	DrawDebugLinesPairsArrayView(Drawer, {
+		{ F0, F1 },
+		{ F1, F2 },
+		{ F2, F3 },
+		{ F3, F0 },
+
+		{ I0, I1 },
+		{ I1, I2 },
+		{ I2, I3 },
+		{ I3, I0 },
+		}, Color, Thickness, bDepthTest);
+
+	if (bOuterHandle)
+	{
+		const FVector LOuterHandle = bHandleOnLeft ?
+			FVector(+HandleRadius + HandleExtension, HandleOffset - Width, HandleHeight) :
+			FVector(+HandleRadius + HandleExtension, Width - HandleOffset, HandleHeight);
+		const FVector OuterHandle = Location + Rotation.RotateVector(OpenRotation.RotateVector(LOuterHandle) + OpenOffset);
+		DrawDebugSphere(Drawer, OuterHandle, FRotator::ZeroRotator, Color, HandleRadius, Thickness, 8, bDepthTest);
+	}
+
+	if (bInnerHandle)
+	{
+		const FVector LInnerHandle = bHandleOnLeft ?
+			FVector(-HandleRadius - HandleExtension, HandleOffset - Width, HandleHeight) :
+			FVector(-HandleRadius - HandleExtension, Width - HandleOffset, HandleHeight);
+		const FVector InnerHandle = Location + Rotation.RotateVector(OpenRotation.RotateVector(LInnerHandle) + OpenOffset);
+		DrawDebugSphere(Drawer, InnerHandle, FRotator::ZeroRotator, Color, HandleRadius, Thickness, 8, bDepthTest);
+	}
+
+	if (bDrawEntryArrow)
+	{
+		DrawDebugArrow(
+			Drawer,
+			Location + Rotation.RotateVector(FVector(-50.0f, 0.0f, 2.0f)),
+			Location + Rotation.RotateVector(FVector(-10.0f, 0.0f, 2.0f)),
+			Color,
+			20.0f,
+			Thickness,
+			bDepthTest);
+	}
+}
+
+void UDrawDebugLibrary::DrawDebugCamera(const FDebugDrawer& Drawer, const FVector& Location, const FRotator& Rotation, float Scale, float FOVDeg, const FLinearColor Color, const float Thickness, const bool bDepthTest)
+{
+	DrawDebugRotation(Drawer, Location, Rotation, Scale, Thickness, Color.A, bDepthTest);
+	DrawDebugBox(Drawer, Location, Rotation, FVector(Scale, Scale, Scale), Color, Thickness, bDepthTest);
+
+	// draw "lens" portion
+	const FRotationTranslationMatrix Axes(Rotation, Location);
+	const FVector XAxis = Axes.GetScaledAxis(EAxis::X);
+	const FVector YAxis = Axes.GetScaledAxis(EAxis::Y);
+	const FVector ZAxis = Axes.GetScaledAxis(EAxis::Z);
+
+	const FVector LensPoint = Location + XAxis * Scale;
+	const float HalfLensSize = Scale * FMath::Tan(FMath::DegreesToRadians(FOVDeg * 0.5f));
+
+	const FVector C0 = LensPoint + XAxis * Scale+ (YAxis * HalfLensSize) + (ZAxis * HalfLensSize);
+	const FVector C1 = LensPoint + XAxis * Scale+ (YAxis * HalfLensSize) - (ZAxis * HalfLensSize);
+	const FVector C2 = LensPoint + XAxis * Scale- (YAxis * HalfLensSize) - (ZAxis * HalfLensSize);
+	const FVector C3 = LensPoint + XAxis * Scale- (YAxis * HalfLensSize) + (ZAxis * HalfLensSize);
+
+	DrawDebugLinesPairsArrayView(Drawer, {
+		{ LensPoint, C0 },
+		{ LensPoint, C1 },
+		{ LensPoint, C2 },
+		{ LensPoint, C3 },
+
+		{ C0, C1 },
+		{ C1, C2 },
+		{ C2, C3 },
+		{ C3, C0 },
+		}, Color, Thickness, bDepthTest);
+}
+
+void UDrawDebugLibrary::DrawDebugTrajectory(const FDebugDrawer& Drawer, const TArray<FVector>& Locations, const TArray<FVector>& Directions, const FTransform& RelativeTransform, const FLinearColor Color, const float DrawArrowLength, const float PointRadius, const float ArrowHeadScale, const float Thickness, const int32 Segments, const float VerticalOffset, const bool bDepthTest)
+{
+	DrawDebugTrajectoryFromArrayViews(Drawer, Locations, Directions, RelativeTransform, Color, DrawArrowLength, PointRadius, ArrowHeadScale, Thickness, Segments, VerticalOffset, bDepthTest);
+}
+
+void UDrawDebugLibrary::DrawDebugTrajectoryFromArrayViews(const FDebugDrawer& Drawer, const TArrayView<const FVector> Locations, const TArrayView<const FVector> Directions, const FTransform& RelativeTransform, const FLinearColor Color, const float DrawArrowLength, const float PointRadius, const float ArrowHeadScale, const float Thickness, const int32 Segments, const float VerticalOffset, const bool bDepthTest)
+{
+	const int32 SampleNum = FMath::Min(Locations.Num(), Directions.Num());
+
+	for (int32 SampleIdx = 0; SampleIdx < SampleNum; SampleIdx++)
+	{
+		const FVector Location = RelativeTransform.TransformPosition(Locations[SampleIdx]) + FVector(0.0f, 0.0f, VerticalOffset);
+		const FVector Direction = RelativeTransform.TransformVectorNoScale(Directions[SampleIdx]);
+
+		DrawDebugSphere(Drawer, Location, FRotator::ZeroRotator, Color, PointRadius, Thickness, Segments, bDepthTest);
+		DrawDebugArrow(Drawer, Location, Location + DrawArrowLength * Direction, Color, 0.2f * ArrowHeadScale * DrawArrowLength, Thickness, bDepthTest);
+
+		if (SampleIdx > 0)
+		{
+			const FVector PrevLocation = RelativeTransform.TransformPosition(Locations[SampleIdx - 1]) + FVector(0.0f, 0.0f, VerticalOffset);
+			DrawDebugLine(Drawer, Location, PrevLocation, Color, Thickness, bDepthTest);
+		}
+	}
+}
+
+void UDrawDebugLibrary::DrawDebugTransformTrajectory(const FDebugDrawer& Drawer, const FTransformTrajectory& TransformTrajectory, const FLinearColor Color, const float Radius, const float Thickness, const float VerticalOffset, const bool bDepthTest)
+{
+
+	FTransformTrajectorySample TimeZeroSample = TransformTrajectory.GetSampleAtTime(0.0f);
+	DrawDebugSphere(Drawer, TimeZeroSample.Position + FVector(0.0f, 0.0f, VerticalOffset), FRotator::ZeroRotator, Color, Radius, Thickness, 8, bDepthTest);
+
+	const int32 SampleNum = TransformTrajectory.Samples.Num();
+	for (int32 SampleIdx = 0; SampleIdx < SampleNum; SampleIdx++)
+	{
+		const FVector Location = TransformTrajectory.Samples[SampleIdx].Position + FVector(0.0f, 0.0f, VerticalOffset);
+		const FQuat Rotation = TransformTrajectory.Samples[SampleIdx].Facing;
+
+		DrawDebugRotation(Drawer, Location, Rotation.Rotator(), Radius, Thickness, Color.A, bDepthTest);
+		DrawDebugString(Drawer, FString::Printf(TEXT("%5.2f"), TransformTrajectory.Samples[SampleIdx].TimeInSeconds), Location + FVector(0.0f, 0.0, 10.0f), FRotator::ZeroRotator, 10.0f, true, Color);
+
+		if (SampleIdx > 0)
+		{
+			const FVector PrevLocation = TransformTrajectory.Samples[SampleIdx - 1].Position + FVector(0.0f, 0.0f, VerticalOffset);
+			DrawDebugLine(Drawer, Location, PrevLocation, Color, Thickness, bDepthTest);
+		}
+	}
+}
+
+void UDrawDebugLibrary::DrawDebugPose(const FDebugDrawer& Drawer, const TArray<FVector>& BoneLocations, const TArray<FVector>& BoneLinearVelocities, const FTransform& RelativeTransform, const FLinearColor Color, const float DrawVelocityLineScale, const float Thickness, const bool bDepthTest)
+{
+	DrawDebugPoseFromArrayViews(Drawer, BoneLocations, BoneLinearVelocities, RelativeTransform, Color, DrawVelocityLineScale, Thickness, bDepthTest);
+}
+
+void UDrawDebugLibrary::DrawDebugPoseFromArrayViews(const FDebugDrawer& Drawer, const TArrayView<const FVector> BoneLocations, const TArrayView<const FVector> BoneLinearVelocities, const FTransform& RelativeTransform, const FLinearColor Color, const float DrawVelocityLineScale, const float Thickness, const bool bDepthTest)
+{
+	const int32 BoneNum = FMath::Min(BoneLocations.Num(), BoneLinearVelocities.Num());
+
+	for (int32 BoneIdx = 0; BoneIdx < BoneNum; BoneIdx++)
+	{
+		const FVector Location = RelativeTransform.TransformPosition(BoneLocations[BoneIdx]);
+		const FVector Velocity = RelativeTransform.TransformVector(BoneLinearVelocities[BoneIdx]);
+
+		DrawDebugSphere(Drawer, Location, FRotator::ZeroRotator, Color, 10.0f, Thickness, 8, bDepthTest);
+		DrawDebugLine(Drawer, Location, Location + DrawVelocityLineScale * Velocity, Color, Thickness, bDepthTest);
+	}
+}
+
+
+void UDrawDebugLibrary::DrawDebugVelocities(const FDebugDrawer& Drawer, const TArray<FVector>& Locations, const TArray<FVector>& Velocities, const FLinearColor Color, const float DrawVelocityLineScale, const float Thickness, const bool bDepthTest)
+{
+	DrawDebugVelocitiesArrayView(Drawer, Locations, Velocities, Color, DrawVelocityLineScale, Thickness, bDepthTest);
+}
+
+void UDrawDebugLibrary::DrawDebugVelocitiesArrayView(const FDebugDrawer& Drawer, const TArrayView<const FVector> Locations, const TArrayView<const FVector> Velocities, const FLinearColor Color, const float DrawVelocityLineScale, const float Thickness, const bool bDepthTest)
+{
+	const int32 VelocityNum = FMath::Min(Locations.Num(), Velocities.Num());
+
+	for (int32 VelocityIdx = 0; VelocityIdx < VelocityNum; VelocityIdx++)
+	{
+		DrawDebugVelocity(Drawer, Locations[VelocityIdx], Velocities[VelocityIdx], Color, DrawVelocityLineScale, Thickness, bDepthTest);
+	}
+}
+
+
+void UDrawDebugLibrary::DrawDebugTransformsArrayView(
+	const FDebugDrawer& Drawer,
+	const TArrayView<const FVector> Locations,
+	const TArrayView<const FQuat4f> Rotations,
+	const TArrayView<const FVector3f> Scales,
+	const float DrawRadius,
+	const float Thickness,
+	const float Opacity,
+	const bool bDepthTest)
+{
+	check(Locations.Num() == Rotations.Num());
+	check(Locations.Num() == Scales.Num());
+
+	const int32 TransformNum = Locations.Num();
+	for (int32 TransformIdx = 0; TransformIdx < TransformNum; TransformIdx++)
+	{
+		const FTransform Transform = FTransform(
+			((FQuat)Rotations[TransformIdx]).GetNormalized(),
+			Locations[TransformIdx],
+			(FVector)Scales[TransformIdx]);
+
+		DrawDebugTransform(
+			Drawer,
+			Transform,
+			DrawRadius,
+			Thickness,
+			Opacity,
+			bDepthTest);
+	}
+}
+
+void UDrawDebugLibrary::DrawDebugRotationsQuatArrayView(
+	const FDebugDrawer& Drawer,
+	const TArrayView<const FVector> Locations,
+	const TArrayView<const FQuat4f> Rotations,
+	const float DrawRadius,
+	const float Thickness,
+	const float Opacity,
+	const bool bDepthTest)
+{
+	const int32 Num = FMath::Min(Locations.Num(), Rotations.Num());
+
+	for (int32 Idx = 0; Idx < Num; Idx++)
+	{
+		DrawDebugRotation(
+			Drawer,
+			Locations[Idx],
+			(FRotator)Rotations[Idx].Rotator(),
+			DrawRadius,
+			Thickness,
+			Opacity,
+			bDepthTest);
+	}
+}
+
+void UDrawDebugLibrary::DrawDebugSkeletonArrayView(
+	const FDebugDrawer& Drawer,
+	const FVector RootLocation,
+	const TArrayView<const FVector> BoneLocations,
+	const TArrayView<const FQuat4f> BoneRotations,
+	const TArrayView<const FVector3f> BoneScales,
+	const TArrayView<const int32> BoneIndices,
+	const FReferenceSkeleton& ReferenceSkeleton,
+	const FLinearColor& Color,
+	const float BoneDrawSize)
+{
+	// Skeleton Drawing only possible for PDI
+	if (!Drawer.PDI) { return; }
+
+	const int32 BoneNum = BoneLocations.Num();
+
+	TArray<FTransform> WorldTransforms;
+	TArray<FLinearColor> BoneColorsArray;
+	WorldTransforms.SetNumUninitialized(BoneNum);
+	BoneColorsArray.SetNumUninitialized(BoneNum);
+	for (int32 BoneIdx = 0; BoneIdx < BoneNum; BoneIdx++)
+	{
+		WorldTransforms[BoneIdx] = FTransform(
+			((FQuat)BoneRotations[BoneIdx]).GetNormalized(),
+			BoneLocations[BoneIdx],
+			(FVector)BoneScales[BoneIdx]);
+		BoneColorsArray[BoneIdx] = Color;
+	}
+
+	TArray<FBoneIndexType> BoneIndices16;
+	BoneIndices16.SetNumUninitialized(BoneIndices.Num());
+	for (int32 BoneIdx = 0; BoneIdx < BoneIndices.Num(); BoneIdx++)
+	{
+		BoneIndices16[BoneIdx] = BoneIndices[BoneIdx];
+	}
+
+	FSkelDebugDrawConfig DrawConfig;
+	DrawConfig.BoneDrawMode = EBoneDrawMode::Type::All;
+	DrawConfig.BoneDrawSize = BoneDrawSize;
+	DrawConfig.bAddHitProxy = false;
+	DrawConfig.bForceDraw = false;
+	DrawConfig.bUseMultiColorAsDefaultColor = false;
+	DrawConfig.DefaultBoneColor = Color;
+	DrawConfig.AffectedBoneColor = Color;
+	DrawConfig.SelectedBoneColor = Color;
+	DrawConfig.ParentOfSelectedBoneColor = Color;
+
+	SkeletalDebugRendering::DrawBones(Drawer.PDI, RootLocation, BoneIndices16, ReferenceSkeleton,
+		WorldTransforms, {}, BoneColorsArray, {}, DrawConfig);
+}
+
+void UDrawDebugLibrary::DrawDebugSimpleSkeletonArrayView(
+	const FDebugDrawer& Drawer,
+	const TArrayView<const FVector> BoneLocations,
+	const TArrayView<const int32> BoneIndices,
+	const TArrayView<const int32> BoneParents,
+	const FLinearColor& Color,
+	const float Thickness,
+	const bool bDepthTest)
+{
+	TArray<TPair<FVector, FVector>> LineSegments;
+	LineSegments.Reserve(BoneIndices.Num());
+
+	for (int32 BoneIdx : BoneIndices)
+	{
+		const int32 ParentIdx = BoneParents[BoneIdx];
+
+		if (ParentIdx != -1)
+		{
+			LineSegments.Add({ BoneLocations[BoneIdx], BoneLocations[ParentIdx] });
+		}
+	}
+
+	DrawDebugLinesPairsArrayView(Drawer, LineSegments, Color, Thickness, bDepthTest);
+}
+
+void UDrawDebugLibrary::DrawDebugBoneVelocitiesArrayView(
+	const FDebugDrawer& Drawer,
+	const TArrayView<const FVector> BoneLocations,
+	const TArrayView<const FVector3f> BoneVelocities,
+	const TArrayView<const int32> BoneIndices,
+	const FLinearColor& Color,
+	const float Scale,
+	const float Thickness,
+	const bool bDepthTest)
+{
+	TArray<TPair<FVector, FVector>> LineSegments;
+	LineSegments.Reserve(BoneIndices.Num());
+
+	for (const int32 BoneIdx : BoneIndices)
+	{
+		LineSegments.Add({ BoneLocations[BoneIdx], BoneLocations[BoneIdx] + Scale * (FVector)BoneVelocities[BoneIdx] });
+	}
+
+	DrawDebugLinesPairsArrayView(Drawer, LineSegments, Color, Thickness, bDepthTest);
+}
+
+void UDrawDebugLibrary::DrawDebugRangeTrajectoryArrayView(
+	const FDebugDrawer& Drawer,
+	const TArrayView<const FVector> TrajectoryLocations,
+	const TArrayView<const FQuat4f> TrajectoryRotations,
+	const FVector& ForwardVector,
+	const bool bDrawOrientations,
+	const bool bDrawStartingAtOrigin,
+	const FVector& OriginOffset,
+	const bool bDepthTest)
+{
+	const int32 FrameNum = TrajectoryLocations.Num();
+	if (FrameNum == 0) { return; }
+
+	TArray<TPair<FVector, FVector>> LineSegments;
+	LineSegments.Reserve(FrameNum);
+
+	TArray<FVector> Points;
+	Points.Reserve(FrameNum);
+
+	for (int32 FrameIdx = 0; FrameIdx < FrameNum; FrameIdx++)
+	{
+		FVector CurrLocation = TrajectoryLocations[FrameIdx];
+		if (bDrawStartingAtOrigin)
+		{
+			CurrLocation = OriginOffset + ((FQuat)TrajectoryRotations[0]).GetNormalized().UnrotateVector(CurrLocation - TrajectoryLocations[0]);
+		}
+
+		const bool bFirstOrLastPoint = FrameIdx == 0 || FrameIdx == FrameNum - 1;
+
+		if (bFirstOrLastPoint)
+		{
+			DrawDebugPoint(Drawer, CurrLocation, FColor::Black.WithAlpha(64), 2.5f, bDepthTest);
+		}
+		else
+		{
+			Points.Add(CurrLocation);
+		}
+
+		if (bDrawOrientations)
+		{
+			if (bFirstOrLastPoint || (FrameIdx % 3 == 0))
+			{
+				FVector CurrDirection = ((FQuat)TrajectoryRotations[FrameIdx]).GetNormalized().RotateVector(ForwardVector);
+				if (bDrawStartingAtOrigin)
+				{
+					CurrDirection = ((FQuat)TrajectoryRotations[0]).GetNormalized().UnrotateVector(CurrDirection);
+				}
+
+				DrawDebugFlatArrow(
+					Drawer,
+					CurrLocation,
+					CurrDirection.Rotation(),
+					FColor::Black.WithAlpha(64),
+					15.0f, 8.0f, 1.0f, bDepthTest);
+			}
+		}
+
+		if (FrameIdx < FrameNum - 1)
+		{
+			FVector NextLocation = TrajectoryLocations[FrameIdx + 1];
+			if (bDrawStartingAtOrigin)
+			{
+				NextLocation = OriginOffset + ((FQuat)TrajectoryRotations[0]).GetNormalized().UnrotateVector(NextLocation - TrajectoryLocations[0]);
+			}
+
+			LineSegments.Add({ CurrLocation, NextLocation });
+		}
+	}
+
+	DrawDebugPointsArrayView(Drawer, Points, FColor::Black.WithAlpha(64), 1.25f, bDepthTest);
+	DrawDebugLinesPairsArrayView(Drawer, LineSegments, FColor::Black.WithAlpha(64), 1.0f, bDepthTest);
+}
+
+void UDrawDebugLibrary::DrawDebugSkeletonCapsulesArrayView(
+	const FDebugDrawer& Drawer,
+	const TArrayView<const FVector> BoneLocations,
+	const TArrayView<const FQuat4f> BoneRotations,
+	const TArrayView<const FVector3f> BoneLocalLocations,
+	const TArrayView<const int32> BoneIndices,
+	const TArrayView<const int32> BoneParents,
+	const float MinCapsuleRadius,
+	const float MaxCapsuleRadius,
+	const FLinearColor Color,
+	const float Thickness,
+	const int32 Segments,
+	const bool bDepthTest)
+{
+	for (const int32 BoneIdx : BoneIndices)
+	{
+		const int32 ParentBoneIdx = BoneParents[BoneIdx];
+		if (ParentBoneIdx != -1)
+		{
+			const FVector Location = (BoneLocations[BoneIdx] + BoneLocations[ParentBoneIdx]) / 2.0f;
+			const FQuat4f CapsuleRotation = FQuat4f::FindBetweenNormals(
+				FVector3f(0.0f, 0.0f, 1.0f),
+				BoneLocalLocations[BoneIdx].GetSafeNormal(UE_SMALL_NUMBER, FVector3f::ForwardVector));
+
+			const float Length = (BoneLocations[BoneIdx] - BoneLocations[ParentBoneIdx]).Length();
+			const float Radius = FMath::Clamp(Length / 2.0, MinCapsuleRadius, MaxCapsuleRadius);
+
+			UDrawDebugLibrary::DrawDebugCapsule(
+				Drawer,
+				Location,
+				(FRotator)(BoneRotations[ParentBoneIdx] * CapsuleRotation).Rotator(),
+				Color,
+				Radius,
+				Length * 0.5f,
+				Thickness,
+				Segments,
+				bDepthTest);
+		}
+	}
+}
+
+int32 UDrawDebugLibrary::DrawDebugStringViewSegmentNum(const FStringView String, const bool bMonospaced)
+{
+	const int32 CharNum = String.Len();
+
+	int32 SegmentNum = 0;
+
+	for (int32 CharIdx = 0; CharIdx < CharNum; CharIdx++)
+	{
+		if (String.GetData()[CharIdx] >= '!' && String.GetData()[CharIdx] <= '~')
+		{
+			const int32 TableIdx = String.GetData()[CharIdx] - '!';
+			SegmentNum += bMonospaced ? UE::DrawDebugLibrary::Private::MonoStringLineNums[TableIdx] : UE::DrawDebugLibrary::Private::StringLineNums[TableIdx];
+		}
+	}
+
+	return SegmentNum;
+}
+
+FVector UDrawDebugLibrary::DrawDebugStringDimensions(const FString& String, const float Scale, const bool bMonospaced)
+{
+	return DrawDebugStringViewDimensions(String, Scale, bMonospaced);
+}
+
+FVector UDrawDebugLibrary::DrawDebugStringViewDimensions(const FStringView String, const float Scale, const bool bMonospaced)
+{
+	const int32 CharNum = String.Len();
+
+	float XOffset = 0.0;
+	float YOffset = 0.0;
+
+	for (int32 CharIdx = 0; CharIdx < CharNum; CharIdx++)
+	{
+		if (String.GetData()[CharIdx] == '\n')
+		{
+			XOffset = 0.0;
+			YOffset -= Scale;
+		}
+		else if (String.GetData()[CharIdx] == '\r')
+		{
+			XOffset = 0.0;
+		}
+		else if (String.GetData()[CharIdx] == '\t')
+		{
+			XOffset += 4.0f * 0.5f * Scale;
+		}
+		else if (String.GetData()[CharIdx] < '!' || String.GetData()[CharIdx] > '~')
+		{
+			XOffset += 0.5f * Scale;
+		}
+		else
+		{
+			const int32 TableIdx = String.GetData()[CharIdx] - '!';
+			XOffset += (bMonospaced ? 0.5f : (float)UE::DrawDebugLibrary::Private::StringLineAdvs[TableIdx] / 128) * Scale;
+		}
+	}
+
+	YOffset -= Scale;
+
+	return FVector(XOffset, -YOffset, 0);
+}
+
+FVector UDrawDebugLibrary::DrawDebugTextDimensions(const FText& Text, const float Scale, const bool bMonospaced)
+{
+	return DrawDebugStringViewDimensions(Text.ToString(), Scale, bMonospaced);
+}
+
+FVector UDrawDebugLibrary::DrawDebugNameDimensions(const FName& Name, const float Scale, const bool bMonospaced)
+{
+	return DrawDebugStringViewDimensions(Name.ToString(), Scale, bMonospaced);
+}
+
+
+void UDrawDebugLibrary::DrawDebugString(const FDebugDrawer& Drawer, const FString& String, const FVector& Location, const FRotator& Rotation, const float Scale, const bool bMonospaced, const FLinearColor Color, const float Thickness, const bool bDepthTest)
+{
+	DrawDebugStringView(Drawer, String, Location, Rotation, Scale, bMonospaced, Color, Thickness, bDepthTest);
+}
+
+void UDrawDebugLibrary::DrawDebugStringView(const FDebugDrawer& Drawer, const FStringView String, const FVector& Location, const FRotator& Rotation, const float Scale, const bool bMonospaced, const FLinearColor Color, const float Thickness, const bool bDepthTest)
+{
+	const int32 CharNum = String.Len();
+
+	TArray<TPair<FVector, FVector>> LineSegments;
+	LineSegments.Reserve(DrawDebugStringViewSegmentNum(String, bMonospaced));
+
+	float XOffset = 0.0;
+	float YOffset = 0.0;
+
+	for (int32 CharIdx = 0; CharIdx < CharNum; CharIdx++)
+	{
+		if (String.GetData()[CharIdx] == '\n')
+		{
+			XOffset = 0.0;
+			YOffset -= Scale;
+		}
+		else if (String.GetData()[CharIdx] == '\r')
+		{
+			XOffset = 0.0;
+		}
+		else if (String.GetData()[CharIdx] == '\t')
+		{
+			XOffset += 4.0f * 0.5f * Scale;
+		}
+		else if (String.GetData()[CharIdx] < '!' || String.GetData()[CharIdx] > '~')
+		{
+			XOffset += 0.5f * Scale;
+		}
+		else
+		{
+			const int32 TableIdx = String.GetData()[CharIdx] - '!';
+			const int32 SegmentNum = bMonospaced ? UE::DrawDebugLibrary::Private::MonoStringLineNums[TableIdx] : UE::DrawDebugLibrary::Private::StringLineNums[TableIdx];
+			const int32 SegmentOff = bMonospaced ? UE::DrawDebugLibrary::Private::MonoStringLineOffsets[TableIdx] : UE::DrawDebugLibrary::Private::StringLineOffsets[TableIdx];
+			const int* Segments = bMonospaced ? UE::DrawDebugLibrary::Private::MonoStringLineSegments : UE::DrawDebugLibrary::Private::StringLineSegments;
+
+			for (int32 SegmentIdx = 0; SegmentIdx < SegmentNum; SegmentIdx++)
+			{
+				const int32 XStartInt = (0x000000FF & (Segments[SegmentOff + SegmentIdx] >> 0));
+				const int32 YStartInt = (0x000000FF & (Segments[SegmentOff + SegmentIdx] >> 8));
+				const int32 XEndInt = (0x000000FF & (Segments[SegmentOff + SegmentIdx] >> 16));
+				const int32 YEndInt = (0x000000FF & (Segments[SegmentOff + SegmentIdx] >> 24));
+
+				const FVector Start = FVector(
+					XOffset + Scale * (((float)XStartInt) / 128),
+					0.0f,
+					YOffset - Scale * (1.0f - ((float)YStartInt) / 128));
+
+				const FVector End = FVector(
+					XOffset + Scale * (((float)XEndInt) / 128),
+					0.0f,
+					YOffset - Scale * (1.0f - ((float)YEndInt) / 128));
+
+				LineSegments.Add({ Rotation.RotateVector(Start) + Location, Rotation.RotateVector(End) + Location });
+			}
+
+			XOffset += (bMonospaced ? 0.5f : (float)UE::DrawDebugLibrary::Private::StringLineAdvs[TableIdx] / 128) * Scale;
+		}
+	}
+
+	DrawDebugLinesPairsArrayView(Drawer, LineSegments, Color, Thickness, bDepthTest);
+}
+
+void UDrawDebugLibrary::DrawDebugName(const FDebugDrawer& Drawer, const FName& Name, const FVector& Location, const FRotator& Rotation, const float Scale, const bool bMonospaced, const FLinearColor Color, const float Thickness, const bool bDepthTest)
+{
+	DrawDebugStringView(Drawer, Name.ToString(), Location, Rotation, Scale, bMonospaced, Color, Thickness, bDepthTest);
+}
+
+void UDrawDebugLibrary::DrawDebugText(const FDebugDrawer& Drawer, const FText& Text, const FVector& Location, const FRotator& Rotation, const float Scale, const bool bMonospaced, const FLinearColor Color, const float Thickness, const bool bDepthTest)
+{
+	DrawDebugStringView(Drawer, Text.ToString(), Location, Rotation, Scale, bMonospaced, Color, Thickness, bDepthTest);
+}
+
+
+void UDrawDebugLibrary::DrawDebugVisualLoggerString(const FDebugDrawer& Drawer, const FString& String, const FVector& Location, const FLinearColor Color)
+{
+	DrawDebugVisualLoggerStringView(Drawer, String, Location, Color);
+}
+
+void UDrawDebugLibrary::DrawDebugVisualLoggerStringView(const FDebugDrawer& Drawer, const FStringView String, const FVector& Location, const FLinearColor Color)
+{
+#if ENABLE_VISUAL_LOG
+	if (Drawer.VisualLoggerObject)
+	{
+		FVisualLogEntry* CurrentEntry = nullptr;
+		if (FVisualLogger::CheckVisualLogInputInternal(Drawer.VisualLoggerObject, Drawer.VisualLoggerCategory, Drawer.VisualLoggerVerbosity, /*OutWorld=*/ nullptr, &CurrentEntry))
+		{
+			FVisualLogShapeElement Element(FString(String), Color.ToFColor(false), 0, Drawer.VisualLoggerCategory);
+			Element.Type = EVisualLoggerShapeElement::SinglePoint;
+			Element.Verbosity = Drawer.VisualLoggerVerbosity;
+			Element.Points.Add(Location);
+			CurrentEntry->ElementsToDraw.Add(Element);
+		}
+	}
+#endif
+}
+
+void UDrawDebugLibrary::DrawDebugVisualLoggerName(const FDebugDrawer& Drawer, const FName& Name, const FVector& Location, const FLinearColor Color)
+{
+	DrawDebugVisualLoggerStringView(Drawer, Name.ToString(), Location, Color);
+}
+
+void UDrawDebugLibrary::DrawDebugVisualLoggerText(const FDebugDrawer& Drawer, const FText& Text, const FVector& Location, const FLinearColor Color)
+{
+	DrawDebugVisualLoggerStringView(Drawer, Text.ToString(), Location, Color);
+}
+
+
+void UDrawDebugLibrary::DrawDebugMoverOrientation(const FDebugDrawer& Drawer, const FVector& Location, const FRotator& Rotation, const FVector ForwardVector, const float Scale, const FLinearColor Color, const float Thickness, const bool bDepthTest)
+{
+	DrawDebugRotation(
+		Drawer,
+		Location,
+		Rotation,
+		Scale / 3.0f,
+		Thickness);
+
+	const float ArrowSize = (2.0f / 3.0f) * Scale;
+	const FVector ArrowDirection = Rotation.RotateVector(ForwardVector);
+	const FVector ArrowPosition = Location - (2.0f * ArrowSize / 3.0f) * ArrowDirection;
+
+	DrawDebugFlatArrow(
+		Drawer,
+		ArrowPosition,
+		ArrowDirection.Rotation(),
+		Color,
+		2.0f * ArrowSize,
+		ArrowSize,
+		Thickness);
+
+	DrawDebugCircle(
+		Drawer,
+		Location,
+		FRotator::ZeroRotator,
+		Color,
+		Scale,
+		Thickness,
+		31);
+}
+
+void UDrawDebugLibrary::DrawDebugGraph(
+	const FDebugDrawer& Drawer,
+	const FVector& Location,
+	const FRotator& Rotation,
+	const TArray<float>& Xvalues,
+	const TArray<float>& Yvalues,
+	const float Xmin,
+	const float Xmax,
+	const float Ymin,
+	const float Ymax,
+	const float XaxisLength,
+	const float YaxisLength,
+	const FString& Title,
+	const FString& XaxisLabel,
+	const FString& YaxisLabel,
+	const float TitleScale,
+	const FLinearColor Color,
+	const float AxisThickness,
+	const float LineThickness,
+	const float TextThickness,
+	const bool bDepthTest)
+{
+	DrawDebugGraphView(
+		Drawer,
+		Location,
+		Rotation,
+		Xvalues,
+		Yvalues,
+		Xmin,
+		Xmax,
+		Ymin,
+		Ymax,
+		XaxisLength,
+		YaxisLength,
+		Title,
+		XaxisLabel,
+		YaxisLabel,
+		TitleScale,
+		Color,
+		AxisThickness,
+		LineThickness,
+		TextThickness,
+		bDepthTest);
+}
+
+void UDrawDebugLibrary::DrawDebugGraphView(
+	const FDebugDrawer& Drawer,
+	const FVector& Location,
+	const FRotator& Rotation,
+	const TArrayView<const float> Xvalues,
+	const TArrayView<const float> Yvalues,
+	const float Xmin,
+	const float Xmax,
+	const float Ymin,
+	const float Ymax,
+	const float XaxisLength,
+	const float YaxisLength,
+	const FStringView Title,
+	const FStringView XaxisLabel,
+	const FStringView YaxisLabel,
+	const float TextScale,
+	const FLinearColor Color,
+	const float AxisThickness,
+	const float LineThickness,
+	const float TextThickness,
+	const bool bDepthTest)
+{
+	DrawDebugOrientedArrow(
+		Drawer,
+		Location,
+		(Rotation.Quaternion() * FQuat::MakeFromRotationVector(FVector(UE_PI / 2, 0, 0))).Rotator(),
+		XaxisLength,
+		Color,
+		10.0f,
+		AxisThickness,
+		bDepthTest);
+
+	DrawDebugOrientedArrow(
+		Drawer,
+		Location,
+		(Rotation.Quaternion() * 
+			FQuat::MakeFromRotationVector(FVector(0, -UE_PI/2, 0)) *
+			FQuat::MakeFromRotationVector(FVector(UE_PI / 2, 0, 0))).Rotator(),
+		YaxisLength,
+		Color,
+		10.0f,
+		AxisThickness,
+		bDepthTest);
+
+	const FVector TitleDimension = DrawDebugStringViewDimensions(Title, TextScale, false);
+
+	DrawDebugStringView(
+		Drawer,
+		Title,
+		Location + Rotation.RotateVector(FVector(XaxisLength / 2.0f - TitleDimension.X / 2.0f, 0.0, YaxisLength + TitleDimension.Y + 0.1f * TextScale)),
+		Rotation,
+		TextScale,
+		false,
+		Color,
+		TextThickness,
+		bDepthTest);
+
+	const FVector XlabelDimension = DrawDebugStringViewDimensions(XaxisLabel, TextScale, false);
+
+	DrawDebugStringView(
+		Drawer,
+		XaxisLabel,
+		Location + Rotation.RotateVector(FVector(XaxisLength / 2.0f - XlabelDimension.X / 2.0f, 0.0, -0.1f * TextScale)),
+		Rotation,
+		TextScale,
+		false,
+		Color,
+		TextThickness,
+		bDepthTest);
+
+	const FVector YlabelDimension = DrawDebugStringViewDimensions(YaxisLabel, TextScale, false);
+
+	DrawDebugStringView(
+		Drawer,
+		YaxisLabel,
+		Location + Rotation.RotateVector(FVector(-YlabelDimension.Y - 0.1f * TextScale, 0.0, YaxisLength / 2 - YlabelDimension.X / 2)),
+		(Rotation.Quaternion() * FQuat::MakeFromRotationVector(FVector(0, -UE_PI/2, 0))).Rotator(),
+		TextScale,
+		false,
+		Color,
+		TextThickness,
+		bDepthTest);
+
+	const int32 PointNum = FMath::Min(Xvalues.Num(), Yvalues.Num());
+
+	TArray<TPair<FVector, FVector>> LineSegments;
+	LineSegments.Reserve(PointNum - 1);
+
+	for (int32 PointIdx = 0; PointIdx < PointNum - 1; PointIdx++)
+	{
+		const float Xloc0 = XaxisLength * ((Xvalues[PointIdx + 0] - Xmin) / (Xmax - Xmin));
+		const float Yloc0 = YaxisLength * ((Yvalues[PointIdx + 0] - Ymin) / (Ymax - Ymin));
+		const float Xloc1 = XaxisLength * ((Xvalues[PointIdx + 1] - Xmin) / (Xmax - Xmin));
+		const float Yloc1 = YaxisLength * ((Yvalues[PointIdx + 1] - Ymin) / (Ymax - Ymin));
+
+		LineSegments.Add({
+			Location + Rotation.RotateVector(FVector(Xloc0, 0.0, Yloc0)),
+			Location + Rotation.RotateVector(FVector(Xloc1, 0.0, Yloc1)) });
+	}
+
+	DrawDebugLinesPairsArrayView(Drawer, LineSegments, Color, LineThickness, bDepthTest);
+}
\ No newline at end of file
diff --git a/Engine/Plugins/Experimental/DrawDebugLibrary/Source/DrawDebugLibrary/Public/DrawDebugLibrary.h b/Engine/Plugins/Experimental/DrawDebugLibrary/Source/DrawDebugLibrary/Public/DrawDebugLibrary.h
new file mode 100644
--- /dev/null
+++ b/Engine/Plugins/Experimental/DrawDebugLibrary/Source/DrawDebugLibrary/Public/DrawDebugLibrary.h
@@ -0,0 +1,546 @@
+// Copyright Epic Games, Inc. All Rights Reserved.
+
+#pragma once
+
+#include "Modules/ModuleManager.h"
+#include "EngineDefines.h"
+#include "Kismet/BlueprintFunctionLibrary.h"
+
+#include "DrawDebugLibrary.generated.h"
+
+#define UE_API DRAWDEBUGLIBRARY_API
+
+class FPrimitiveDrawInterface;
+struct FReferenceSkeleton;
+struct FTransformTrajectory;
+struct FAnimInstanceProxy;
+
+class FDrawDebugLibraryModule : public IModuleInterface
+{
+public:
+
+	/** IModuleInterface implementation */
+	virtual void StartupModule() override;
+	virtual void ShutdownModule() override;
+};
+
+/** Blueprint accessible version of log verbosity for the Visual Logger Drawer */
+UENUM(BlueprintType)
+enum class EDrawDebugLogVerbosity : uint8
+{
+	/** Always prints a fatal error to console (and log file) and crashes (even if logging is disabled) */
+	Fatal,
+
+	/**
+	 * Prints an error to console (and log file).
+	 * Commandlets and the editor collect and report errors. Error messages result in commandlet failure.
+	 */
+	Error,
+
+	/**
+	 * Prints a warning to console (and log file).
+	 * Commandlets and the editor collect and report warnings. Warnings can be treated as an error.
+	 */
+	Warning,
+
+	/** Prints a message to console (and log file) */
+	Display,
+
+	/** Prints a message to a log file (does not print to console) */
+	Log,
+
+	/**
+	 * Prints a verbose message to a log file (if Verbose logging is enabled for the given category,
+	 * usually used for detailed logging)
+	 */
+	Verbose,
+
+	/**
+	 * Prints a verbose message to a log file (if VeryVerbose logging is enabled,
+	 * usually used for detailed logging that would otherwise spam output)
+	 */
+	VeryVerbose,
+};
+
+/**
+ * A simple Blueprint wrapper around different objects that are capable of debug drawing. This allows us to use the same blueprint interface and
+ * functions for debug drawing if we are drawing using a UWorld object or a FPrimitiveDrawInterface object (for example). Since this contains raw 
+ * pointers it should not be stored anywhere as these could become stale. It should always be treated as a temporary object.
+ */
+USTRUCT(BlueprintType)
+struct FDebugDrawer
+{
+	GENERATED_BODY()
+
+public:
+
+	/** Primitive Draw Interface - used for drawing in editor viewports. */
+	FPrimitiveDrawInterface* PDI = nullptr;
+
+	/** World object - used for normal debug drawing. */
+	UWorld* World = nullptr;
+
+	/** Anim Instance Proxy Object - used for debug drawing in an AnimNode. */
+	FAnimInstanceProxy* AnimInstanceProxy = nullptr;
+
+	/** Visual Logger Object - used for writing to the visual logger. */
+	UObject* VisualLoggerObject = nullptr;
+
+	/** Visual Logger Category. */
+	FName VisualLoggerCategory = NAME_None;
+
+	/** Visual Logger Verbosity. */
+	ELogVerbosity::Type VisualLoggerVerbosity = ELogVerbosity::Display;
+};
+
+/**
+ * A blueprint library of additional debug draw functions that unifies drawing between the PDI and World Debug Drawing functions.
+ *
+ * The built-in world debug draw functions are only accessible via blueprints for actor types or other classes which have an implicit World associated
+ * with them. Here we provide blueprint-accessible debug draw functions that are accessible from anywhere (as long as you have either a world or
+ * PDI variable accessibly) and we also provide functions for debug drawing a number of additional types of things.
+ */
+UCLASS(BlueprintType)
+class UDrawDebugLibrary : public UBlueprintFunctionLibrary
+{
+	GENERATED_BODY()
+
+public: // Drawers
+
+	/** Make a null debug drawer which ignores draw commands */
+	UFUNCTION(BlueprintPure, Category = "Draw Debug Library")
+	static UE_API FDebugDrawer MakeNullDebugDrawer();
+
+	/** Make a debug debugger from the current "Self" object */
+	UFUNCTION(BlueprintPure, Category = "Draw Debug Library", meta = (DevelopmentOnly, DefaultToSelf = "Object", HidePin = "Object"))
+	static UE_API FDebugDrawer MakeDebugDrawer(UObject* Object);
+
+	/** Make a debug drawer for a UObject */
+	UFUNCTION(BlueprintPure, Category = "Draw Debug Library", meta = (DevelopmentOnly))
+	static UE_API FDebugDrawer MakeObjectDebugDrawer(UObject* Object);
+
+	/** Make a debug drawer for a UWorld */
+	static UE_API FDebugDrawer MakeWorldDebugDrawer(UWorld* World);
+
+	/** Make a debug drawer for an AnimInstanceProxy */
+	static UE_API FDebugDrawer MakeAnimInstanceProxyDebugDrawer(FAnimInstanceProxy* AnimInstanceProxy);
+
+	/** Make a debug drawer for a FPrimitiveDrawInterface */
+	static UE_API FDebugDrawer MakePDIDebugDrawer(FPrimitiveDrawInterface* PrimitiveDrawInterface);
+
+	/** Make a debug drawer for the Visual Logger from the current "Self" object */
+	UFUNCTION(BlueprintPure, Category = "Draw Debug Library", meta = (DevelopmentOnly, DefaultToSelf = "Object", HidePin = "Object"))
+	static UE_API FDebugDrawer MakeVisualLoggerDebugDrawer(UObject* Object, const FName Category = TEXT("LogDrawDebugLibrary"), const EDrawDebugLogVerbosity Verbosity = EDrawDebugLogVerbosity::Display, const bool bDrawToScene = true);
+
+	/** Make a debug drawer for the Visual Logger */
+	UFUNCTION(BlueprintPure, Category = "Draw Debug Library", meta = (DevelopmentOnly))
+	static UE_API FDebugDrawer MakeVisualLoggerDebugDrawerFromObject(UObject* Object, const FName Category = TEXT("LogDrawDebugLibrary"), const EDrawDebugLogVerbosity Verbosity = EDrawDebugLogVerbosity::Display, const bool bDrawToScene = true);
+	static UE_API FDebugDrawer MakeVisualLoggerDebugDrawerFromObjectWithCategory(UObject* Object, const FLogCategoryBase& Category, const ELogVerbosity::Type Verbosity = ELogVerbosity::Display, const bool bDrawToScene = true);
+
+	/** Make a debug drawer by merging multiple other debug drawers of different types */
+	UFUNCTION(BlueprintPure, Category = "Draw Debug Library", meta = (DevelopmentOnly))
+	static UE_API FDebugDrawer MakeMergedDebugDrawer(const TArray<FDebugDrawer>& DebugDrawers);
+	static UE_API FDebugDrawer MakeMergedDebugDrawerArrayView(const TArrayView<const FDebugDrawer> DebugDrawers);
+
+public: // Visual Debugger Logging
+
+	/** Log a string to the visual logger. Only does anything with a VisualLoggerDebugDrawer. */
+	UFUNCTION(BlueprintCallable, Category = "Draw Debug Library", meta = (DevelopmentOnly, AutoCreateRefTerm = "String"))
+	static UE_API void LogVisualLoggerString(const FDebugDrawer& Drawer, const FString& String, const FLinearColor Color = FLinearColor::Red);
+	static UE_API void LogVisualLoggerStringView(const FDebugDrawer& Drawer, const FStringView String, const FLinearColor Color = FLinearColor::Red);
+
+	/** Log some text to the visual logger. Only does anything with a VisualLoggerDebugDrawer. */
+	UFUNCTION(BlueprintCallable, Category = "Draw Debug Library", meta = (DevelopmentOnly, AutoCreateRefTerm = "Text"))
+	static UE_API void LogVisualLoggerText(const FDebugDrawer& Drawer, const FText& Text, const FLinearColor Color = FLinearColor::Red);
+
+public: // Convenience functions
+
+	/** Convenience function for making an array of linearly spaced float values */
+	UFUNCTION(BlueprintPure = true, Category = "Draw Debug Library")
+	static UE_API void MakeLinearlySpacedFloatArray(TArray<float>& OutValues, const float Start = 0.0f, const float Stop = 1.0f, const int32 Num = 10);
+
+	/** Convenience function that adds a float to the end of an array, popping values from the front once the max is reached */
+	UFUNCTION(BlueprintPure = false, Category = "Draw Debug Library")
+	static UE_API void AddToFloatHistoryArray(UPARAM(ref) TArray<float>& InOutValues, const float NewValue = 0.0f, const int32 MaxHistoryNum = 60);
+
+	/** Convenience function that adds a vector to the end of an array, popping values from the front once the max is reached */
+	UFUNCTION(BlueprintPure = false, Category = "Draw Debug Library")
+	static UE_API void AddToVectorHistoryArray(UPARAM(ref) TArray<FVector>& InOutValues, const FVector NewValue = FVector::ZeroVector, const int32 MaxHistoryNum = 60);
+
+public: // Point & Line
+
+	/** Debug Draw a point */
+	UFUNCTION(BlueprintCallable, Category = "Draw Debug Library", meta = (DevelopmentOnly, AutoCreateRefTerm = "Location"))
+	static UE_API void DrawDebugPoint(const FDebugDrawer& Drawer, const FVector& Location, const FLinearColor Color = FLinearColor::Red, const float Size = 0.0f, const bool bDepthTest = true);
+
+	/** Debug Draw points. These should be preferred over DrawDebugPoint where possible as it will batch drawing when required such as when using the visual logger. */
+	UFUNCTION(BlueprintCallable, Category = "Draw Debug Library", meta = (DevelopmentOnly))
+	static UE_API void DrawDebugPoints(const FDebugDrawer& Drawer, const TArray<FVector>& Locations, const FLinearColor Color = FLinearColor::Red, const float Size = 0.0f, const bool bDepthTest = true);
+	static UE_API void DrawDebugPointsArrayView(const FDebugDrawer& Drawer, const TArrayView<const FVector> Locations, const FLinearColor Color = FLinearColor::Red, const float Size = 0.0f, const bool bDepthTest = true);
+
+	/** Debug Draw a line */
+	UFUNCTION(BlueprintCallable, Category = "Draw Debug Library", meta = (DevelopmentOnly, AutoCreateRefTerm = "StartLocation, EndLocation"))
+	static UE_API void DrawDebugLine(const FDebugDrawer& Drawer, const FVector& StartLocation, const FVector& EndLocation, const FLinearColor Color = FLinearColor::Red, const float Thickness = 0.0f, const bool bDepthTest = true);
+
+	/** Debug Draw lines. These should be preferred over DrawDebugLine where possible as it will batch drawing when required such as when using the visual logger. */
+	UFUNCTION(BlueprintCallable, Category = "Draw Debug Library", meta = (DevelopmentOnly))
+	static UE_API void DrawDebugLines(const FDebugDrawer& Drawer, const TArray<FVector>& StartLocations, const TArray<FVector>& EndLocations, const FLinearColor Color = FLinearColor::Red, const float Thickness = 0.0f, const bool bDepthTest = true);
+	static UE_API void DrawDebugLinesArrayView(const FDebugDrawer& Drawer, const TArrayView<const FVector> StartLocations, const TArrayView<const FVector> EndLocations, const FLinearColor Color = FLinearColor::Red, const float Thickness = 0.0f, const bool bDepthTest = true);
+	static UE_API void DrawDebugLinesPairsArrayView(const FDebugDrawer& Drawer, const TArrayView<const TPair<FVector, FVector>> Segments, const FLinearColor Color = FLinearColor::Red, const float Thickness = 0.0f, const bool bDepthTest = true);
+
+public: // Basic Shapes
+
+	/** Debug Draw a triangular base pyramid */
+	UFUNCTION(BlueprintCallable, Category = "Draw Debug Library", meta = (DevelopmentOnly, AutoCreateRefTerm = "Location, Rotation"))
+	static UE_API void DrawDebugTriangularBasePyramid(const FDebugDrawer& Drawer, const FVector& Location, const FRotator& Rotation, const FLinearColor Color = FLinearColor::Red, const float Length = 20.0f, const float Width = 20.0f, const float Thickness = 0.0f, const bool bDepthTest = true);
+
+	/** Debug Draw a square base pyramid */
+	UFUNCTION(BlueprintCallable, Category = "Draw Debug Library", meta = (DevelopmentOnly, AutoCreateRefTerm = "Location, Rotation"))
+	static UE_API void DrawDebugSquareBasePyramid(const FDebugDrawer& Drawer, const FVector& Location, const FRotator& Rotation, const FLinearColor Color = FLinearColor::Red, const float Length = 20.0f, const float Width = 20.0f, const float Thickness = 0.0f, const bool bDepthTest = true);
+
+	/** Debug Draw a cone */
+	UFUNCTION(BlueprintCallable, Category = "Draw Debug Library", meta = (DevelopmentOnly, AutoCreateRefTerm = "Location, Rotation"))
+	static UE_API void DrawDebugCone(const FDebugDrawer& Drawer, const FVector& Location, const FRotator& Rotation, const FLinearColor Color = FLinearColor::Red, const float Length = 20.0f, const float Radius = 10.0f, const int32 Segments = 9, const float Thickness = 0.0f, const bool bDepthTest = true);
+
+	/** Debug Draw an arc around the XY axis */
+	UFUNCTION(BlueprintCallable, Category = "Draw Debug Library", meta = (DevelopmentOnly, AutoCreateRefTerm = "Location, Rotation"))
+	static UE_API void DrawDebugArc(const FDebugDrawer& Drawer, const FVector& Location, const FRotator& Rotation, const float Angle = 360.0f, const FLinearColor Color = FLinearColor::Red, const float Radius = 10.0f, const float Thickness = 0.0f, const int32 Segments = 17, const bool bDepthTest = true);
+
+	/** Debug Draw an circle around the XY axis */
+	UFUNCTION(BlueprintCallable, Category = "Draw Debug Library", meta = (DevelopmentOnly, AutoCreateRefTerm = "Location, Rotation"))
+	static UE_API void DrawDebugCircle(const FDebugDrawer& Drawer, const FVector& Location, const FRotator& Rotation, const FLinearColor Color = FLinearColor::Red, const float Radius = 10.0f, const float Thickness = 0.0f, const int32 Segments = 17, const bool bDepthTest = true);
+
+	/** Debug Draw the outline of a circle (two circles) around the XY axis */
+	UFUNCTION(BlueprintCallable, Category = "Draw Debug Library", meta = (DevelopmentOnly, AutoCreateRefTerm = "Location, Rotation"))
+	static UE_API void DrawDebugCircleOutline(const FDebugDrawer& Drawer, const FVector& Location, const FRotator& Rotation, const FLinearColor Color = FLinearColor::Red, const float InnerRadius = 10.0f, const float OuterRadius = 15.0f, const float Thickness = 0.0f, const int32 Segments = 17, const bool bDepthTest = true);
+
+	/** Debug Draw an axis-aligned cross */
+	UFUNCTION(BlueprintCallable, Category = "Draw Debug Library", meta = (DevelopmentOnly, AutoCreateRefTerm = "Location, Rotation"))
+	static UE_API void DrawDebugLocator(const FDebugDrawer& Drawer, const FVector& Location, const FRotator& Rotation, const FLinearColor Color = FLinearColor::Red, const float DrawRadius = 10.0f, const float Thickness = 0.0f, const bool bDepthTest = true);
+
+	/** Debug Draw an axis-aligned cross rotated by 45 degrees */
+	UFUNCTION(BlueprintCallable, Category = "Draw Debug Library", meta = (DevelopmentOnly, AutoCreateRefTerm = "Location, Rotation"))
+	static UE_API void DrawDebugCrossLocator(const FDebugDrawer& Drawer, const FVector& Location, const FRotator& Rotation, const FLinearColor Color = FLinearColor::Red, const float DrawRadius = 10.0f, const float Thickness = 0.0f, const bool bDepthTest = true);
+
+	/** Debug Draw an oriented box */
+	UFUNCTION(BlueprintCallable, Category = "Draw Debug Library", meta = (DevelopmentOnly, AutoCreateRefTerm = "Location, Rotation"))
+	static UE_API void DrawDebugBox(const FDebugDrawer& Drawer, const FVector& Location, const FRotator& Rotation, const FVector HalfExtents = FVector(10.0f, 10.0f, 10.0f), const FLinearColor Color = FLinearColor::Red, const float Thickness = 0.0f, const bool bDepthTest = true);
+
+	/** Debug Draw a sphere */
+	UFUNCTION(BlueprintCallable, Category = "Draw Debug Library", meta = (DevelopmentOnly, AutoCreateRefTerm = "Location, Rotation"))
+	static UE_API void DrawDebugSphere(const FDebugDrawer& Drawer, const FVector& Location, const FRotator& Rotation, const FLinearColor Color = FLinearColor::Red, const float Radius = 10.0f, const float Thickness = 0.0f, const int32 Segments = 8, const bool bDepthTest = true);
+
+	/** Debug Draw a capsule */
+	UFUNCTION(BlueprintCallable, Category = "Draw Debug Library", meta = (DevelopmentOnly, AutoCreateRefTerm = "Location, Rotation"))
+	static UE_API void DrawDebugCapsule(const FDebugDrawer& Drawer, const FVector& Location, const FRotator& Rotation, const FLinearColor Color = FLinearColor::Red, const float Radius = 10.0f, const float HalfLength = 10.0f, const float Thickness = 0.0f, const int32 Segments = 8, const bool bDepthTest = true);
+
+	/** Debug Draw a Frustum */
+	UFUNCTION(BlueprintCallable, Category = "Draw Debug Library", meta = (DevelopmentOnly, AutoCreateRefTerm = "FrustumToWorld"))
+	static UE_API void DrawDebugFrustum(const FDebugDrawer& Drawer, const FMatrix& FrustumToWorld, const FLinearColor Color = FLinearColor::Red, const float Thickness = 0.0f, const bool bDepthTest = true);
+
+
+public: // Arrows
+
+	/** Debug Draw an arrow */
+	UFUNCTION(BlueprintCallable, Category = "Draw Debug Library", meta = (DevelopmentOnly, AutoCreateRefTerm = "StartLocation, EndLocation"))
+	static UE_API void DrawDebugArrow(const FDebugDrawer& Drawer, const FVector& StartLocation, const FVector& EndLocation, const FLinearColor Color = FLinearColor::Red, const float ArrowHeadSize = 20.0f, const float Thickness = 0.0f, const bool bDepthTest = true);
+
+	/** Debug Draw an arrow with an orientation */
+	UFUNCTION(BlueprintCallable, Category = "Draw Debug Library", meta = (DevelopmentOnly, AutoCreateRefTerm = "Location, Rotation"))
+	static UE_API void DrawDebugOrientedArrow(const FDebugDrawer& Drawer, const FVector& Location, const FRotator& Rotation, const float Length = 100.0f, const FLinearColor Color = FLinearColor::Red, const float ArrowHeadSize = 20.0f, const float Thickness = 0.0f, const bool bDepthTest = true);
+
+	/** Debug Draw an arrow with a four-way head */
+	UFUNCTION(BlueprintCallable, Category = "Draw Debug Library", meta = (DevelopmentOnly, AutoCreateRefTerm = "StartLocation, EndLocation"))
+	static UE_API void DrawDebugFourWayArrow(const FDebugDrawer& Drawer, const FVector& StartLocation, const FVector& EndLocation, const FLinearColor Color = FLinearColor::Red, const float ArrowHeadSize = 20.0f, const float Thickness = 0.0f, const bool bDepthTest = true);
+
+	/** Draws an arrow pointing down at the location surrounded by a circle */
+	UFUNCTION(BlueprintCallable, Category = "Draw Debug Library", meta = (DevelopmentOnly, AutoCreateRefTerm = "Location"))
+	static UE_API void DrawDebugGroundTargetArrow(const FDebugDrawer& Drawer, const FVector& Location, const FLinearColor Color = FLinearColor::Red, const float Length = 100.0f, const float ArrowHeadSize = 20.0f, const float Radius = 10.0f, const int32 Segements = 17, const float Thickness = 0.0f, const bool bDepthTest = true);
+
+	/** Debug Draw an arrow with a triangular arrowhead */
+	UFUNCTION(BlueprintCallable, Category = "Draw Debug Library", meta = (DevelopmentOnly, AutoCreateRefTerm = "StartLocation, EndLocation"))
+	static UE_API void DrawTriangularArrow(const FDebugDrawer& Drawer, const FVector& StartLocation, const FVector& EndLocation, const FLinearColor Color = FLinearColor::Red, const float ArrowHeadSize = 20.0f, const float Thickness = 0.0f, const bool bDepthTest = true);
+
+	/** Debug Draw an arrow with a square arrowhead */
+	UFUNCTION(BlueprintCallable, Category = "Draw Debug Library", meta = (DevelopmentOnly, AutoCreateRefTerm = "StartLocation, EndLocation"))
+	static UE_API void DrawSquareArrow(const FDebugDrawer& Drawer, const FVector& StartLocation, const FVector& EndLocation, const FLinearColor Color = FLinearColor::Red, const float ArrowHeadSize = 20.0f, const float Thickness = 0.0f, const bool bDepthTest = true);
+
+	/** Debug Draw an arrow with a cone arrowhead */
+	UFUNCTION(BlueprintCallable, Category = "Draw Debug Library", meta = (DevelopmentOnly, AutoCreateRefTerm = "StartLocation, EndLocation"))
+	static UE_API void DrawConeArrow(const FDebugDrawer& Drawer, const FVector& StartLocation, const FVector& EndLocation, const FLinearColor Color = FLinearColor::Red, const float ArrowHeadSize = 20.0f, const float Thickness = 0.0f, const int32 Segments = 9, const bool bDepthTest = true);
+
+	/** Draws a flat 2d arrow motif */
+	UFUNCTION(BlueprintCallable, Category = "Draw Debug Library", meta = (DevelopmentOnly, AutoCreateRefTerm = "Location, Rotation"))
+	static UE_API void DrawDebugFlatArrow(const FDebugDrawer& Drawer, const FVector& Location, const FRotator& Rotation, const FLinearColor Color = FLinearColor::Red, const float Length = 20.0f, const float Width = 10.0f, const float Thickness = 0.0f, const bool bDepthTest = true);
+
+public: // Splines
+
+	/** Debug Draw the start section of a catmull rom spline */
+	UFUNCTION(BlueprintCallable, Category = "Draw Debug Library", meta = (DevelopmentOnly, AutoCreateRefTerm = "V0, V1, V2"))
+	static UE_API void DrawDebugCatmullRomSplineStart(const FDebugDrawer& Drawer, const FVector& V0, const FVector& V1, const FVector& V2, const bool bMonotonic = false, const FLinearColor Color = FLinearColor::Red, const float Thickness = 0.0f, const int32 Segments = 15, const bool bDepthTest = true);
+
+	/** Debug Draw the end section of a catmull rom spline */
+	UFUNCTION(BlueprintCallable, Category = "Draw Debug Library", meta = (DevelopmentOnly, AutoCreateRefTerm = "V0, V1, V2"))
+	static UE_API void DrawDebugCatmullRomSplineEnd(const FDebugDrawer& Drawer, const FVector& V0, const FVector& V1, const FVector& V2, const bool bMonotonic = false, const FLinearColor Color = FLinearColor::Red, const float Thickness = 0.0f, const int32 Segments = 15, const bool bDepthTest = true);
+
+	/** Debug Draw a full segment of a catmull rom spline */
+	UFUNCTION(BlueprintCallable, Category = "Draw Debug Library", meta = (DevelopmentOnly, AutoCreateRefTerm = "V0, V1, V2, V3"))
+	static UE_API void DrawDebugCatmullRomSplineSection(const FDebugDrawer& Drawer, const FVector& V0, const FVector& V1, const FVector& V2, const FVector& V3, const bool bMonotonic = false, const FLinearColor Color = FLinearColor::Red, const float Thickness = 0.0f, const int32 Segments = 15, const bool bDepthTest = true);
+
+	/** Debug Draw a catmull rom spline */
+	UFUNCTION(BlueprintCallable, Category = "Draw Debug Library", meta = (DevelopmentOnly))
+	static UE_API void DrawDebugCatmullRomSpline(const FDebugDrawer& Drawer, const TArray<FVector>& Points, const bool bMonotonic = false, const FLinearColor Color = FLinearColor::Red, const float Thickness = 0.0f, const int32 Segments = 15, const bool bDepthTest = true);
+	static UE_API void DrawDebugCatmullRomSplineArrayView(const FDebugDrawer& Drawer, const TArrayView<const FVector> Points, const bool bMonotonic = false, const FLinearColor Color = FLinearColor::Red, const float Thickness = 0.0f, const int32 Segments = 15, const bool bDepthTest = true);
+
+public: // Variable Types
+
+	/** Debug Draw an angle */
+	UFUNCTION(BlueprintCallable, Category = "Draw Debug Library", meta = (DevelopmentOnly, AutoCreateRefTerm = "Location, Rotation"))
+	static UE_API void DrawDebugAngle(const FDebugDrawer& Drawer, const FVector& Location, const FRotator& Rotation, const float Angle = 0.0f, const float LineLength = 10.0f, const float AngleRadius = 5.0f, const FLinearColor Color = FLinearColor::Red, const float Thickness = 0.0f, const int32 Segments = 17, const bool bDepthTest = true);
+
+	/** Debug Draw a sphere at the given location */
+	UFUNCTION(BlueprintCallable, Category = "Draw Debug Library", meta = (DevelopmentOnly, AutoCreateRefTerm = "Location"))
+	static UE_API void DrawDebugLocation(const FDebugDrawer& Drawer, const FVector& Location, const FLinearColor Color = FLinearColor::Red, const float DrawRadius = 10.0f, const float Thickness = 0.0f, const int32 Segments = 8, const bool bDepthTest = true);
+
+	/** Debug Draw spheres at the given locations */
+	UFUNCTION(BlueprintCallable, Category = "Draw Debug Library", meta = (DevelopmentOnly))
+	static UE_API void DrawDebugLocations(const FDebugDrawer& Drawer, const TArray<FVector>& Locations, const FLinearColor Color = FLinearColor::Red, const float DrawRadius = 10.0f, const float Thickness = 0.0f, const int32 Segments = 8, const bool bDepthTest = true);
+	static UE_API void DrawDebugLocationsArrayView(const FDebugDrawer& Drawer, const TArrayView<const FVector> Locations, const FLinearColor Color = FLinearColor::Red, const float DrawRadius = 10.0f, const float Thickness = 0.0f, const int32 Segments = 8, const bool bDepthTest = true);
+
+	/** Draws a rotation */
+	UFUNCTION(BlueprintCallable, Category = "Draw Debug Library", meta = (DevelopmentOnly, AutoCreateRefTerm = "Location, Rotation"))
+	static UE_API void DrawDebugRotation(const FDebugDrawer& Drawer, const FVector& Location, const FRotator& Rotation, const float DrawRadius = 10.0f, const float Thickness = 0.0f, const float Opacity = 1.0f, const bool bDepthTest = true);
+
+	/** Draws an array of transforms */
+	static UE_API void DrawDebugRotationsQuatArrayView(
+		const FDebugDrawer& Drawer,
+		const TArrayView<const FVector> Locations,
+		const TArrayView<const FQuat4f> Rotations,
+		const float DrawRadius = 10.0f,
+		const float Thickness = 0.0f,
+		const float Opacity = 1.0f,
+		const bool bDepthTest = true);
+
+	/** Debug Draw an arrow at the given location facing in the given direction */
+	UFUNCTION(BlueprintCallable, Category = "Draw Debug Library", meta = (DevelopmentOnly, AutoCreateRefTerm = "Location, Direction"))
+	static UE_API void DrawDebugDirection(const FDebugDrawer& Drawer, const FVector& Location, const FVector& Direction = FVector::ForwardVector, const FLinearColor Color = FLinearColor::Red, const float DrawArrowLength = 100.0f, const float ArrowHeadScale = 1.0f, const float Thickness = 0.0f, const bool bDepthTest = true);
+
+	/** Debug Draw a line at the given location scaled by the given velocity */
+	UFUNCTION(BlueprintCallable, Category = "Draw Debug Library", meta = (DevelopmentOnly, AutoCreateRefTerm = "Location, Velocity"))
+	static UE_API void DrawDebugVelocity(const FDebugDrawer& Drawer, const FVector& Location, const FVector& Velocity, const FLinearColor Color = FLinearColor::Red, const float DrawVelocityLineScale = 1.0f, const float Thickness = 0.0f, const bool bDepthTest = true);
+
+	/** Debug draw velocities */
+	UFUNCTION(BlueprintCallable, Category = "Draw Debug Library", meta = (DevelopmentOnly))
+	static UE_API void DrawDebugVelocities(const FDebugDrawer& Drawer, const TArray<FVector>& Locations, const TArray<FVector>& Velocities, const FLinearColor Color = FLinearColor::Red, const float DrawVelocityLineScale = 1.0f, const float Thickness = 0.0f, const bool bDepthTest = true);
+	static UE_API void DrawDebugVelocitiesArrayView(const FDebugDrawer& Drawer, const TArrayView<const FVector> Locations, const TArrayView<const FVector> Velocities, const FLinearColor Color = FLinearColor::Red, const float DrawVelocityLineScale = 1.0f, const float Thickness = 0.0f, const bool bDepthTest = true);
+
+	/** Debug Draw a set of axes at the given transform */
+	UFUNCTION(BlueprintCallable, Category = "Draw Debug Library", meta = (DevelopmentOnly, AutoCreateRefTerm = "Transform"))
+	static UE_API void DrawDebugTransform(const FDebugDrawer& Drawer, const FTransform& Transform, const float DrawRadius = 10.0f, const float Thickness = 0.0f, const float Opacity = 1.0f, const bool bDepthTest = true);
+
+	/** Draws an array of transforms */
+	static UE_API void DrawDebugTransformsArrayView(
+		const FDebugDrawer& Drawer,
+		const TArrayView<const FVector> Locations,
+		const TArrayView<const FQuat4f> Rotations,
+		const TArrayView<const FVector3f> Scales,
+		const float DrawRadius = 10.0f,
+		const float Thickness = 0.0f,
+		const float Opacity = 1.0f,
+		const bool bDepthTest = true);
+
+	/** Debug draw a phase-representation of an event at a location */
+	UFUNCTION(BlueprintCallable, Category = "Draw Debug Library", meta = (DevelopmentOnly, AutoCreateRefTerm = "Location"))
+	static UE_API void DrawDebugEvent(const FDebugDrawer& Drawer, bool bTimeUntilEventKnown, const float TimeUntilEvent, const FVector& Location, const FLinearColor Color = FLinearColor::Red, const float Size = 20.0f, const float Thickness = 0.0f, const bool bDepthTest = true);
+
+public: // Complex Objects
+
+	/** Debug draw a basic chair shape */
+	UFUNCTION(BlueprintCallable, Category = "Draw Debug Library", meta = (DevelopmentOnly, AutoCreateRefTerm = "Location, Rotation"))
+	static UE_API void DrawDebugChair(const FDebugDrawer& Drawer, const FVector& Location, const FRotator& Rotation, const float Width = 40.0f, const float SeatHeight = 40.0f, const float BackHeight = 40.0f, const float BackTilt = 10.0f, const FLinearColor Color = FLinearColor::Red, const float Thickness = 0.0f, const bool bDepthTest = true);
+
+	/** Debug draw a basic door shape */
+	UFUNCTION(BlueprintCallable, Category = "Draw Debug Library", meta = (DevelopmentOnly, AutoCreateRefTerm = "Location, Rotation"))
+	static UE_API void DrawDebugDoor(const FDebugDrawer& Drawer, const FVector& Location, const FRotator& Rotation, const float OpenAngle = 0.0f, const float Width = 100.0f, const float Height = 200.0f, const float Inset = 5.0f, const float HandleHeight = 90.0f, const float HandleOffset = 20.0f, const float HandleRadius = 3.0f, const float HandleExtension = 2.0f, const bool bOuterHandle = true, const bool bInnerHandle = true, const bool bHandleOnLeft = false, const bool bDrawEntryArrow = true, const FLinearColor Color = FLinearColor::Red, const float Thickness = 0.0f, const bool bDepthTest = true);
+
+	/** Debug draw a camera shape */
+	UFUNCTION(BlueprintCallable, Category = "Draw Debug Library", meta = (DevelopmentOnly, AutoCreateRefTerm = "Location, Rotation"))
+	static UE_API void DrawDebugCamera(const FDebugDrawer& Drawer, const FVector& Location, const FRotator& Rotation, float Scale = 10.0f, float FOVDeg = 30.0f, const FLinearColor Color = FLinearColor::Red, const float Thickness = 0.0f, const bool bDepthTest = true);
+
+	/** Debug Draw a trajectory of locations and directions */
+	UFUNCTION(BlueprintCallable, Category = "Draw Debug Library", meta = (DevelopmentOnly, AutoCreateRefTerm = "RelativeTransform"))
+	static UE_API void DrawDebugTrajectory(const FDebugDrawer& Drawer, const TArray<FVector>& Locations, const TArray<FVector>& Directions, const FTransform& RelativeTransform = FTransform(), const FLinearColor Color = FLinearColor::Red, const float DrawArrowLength = 100.0f, const float PointRadius = 5.0f, const float ArrowHeadScale = 1.0f, const float Thickness = 0.0f, const int32 Segments = 9, const float VerticalOffset = 2.0f, const bool bDepthTest = true);
+	static UE_API void DrawDebugTrajectoryFromArrayViews(const FDebugDrawer& Drawer, const TArrayView<const FVector> Locations, const TArrayView<const FVector> Directions, const FTransform& RelativeTransform = FTransform(), const FLinearColor Color = FLinearColor::Red, const float DrawArrowLength = 100.0f, const float PointRadius = 5.0f, const float ArrowHeadScale = 1.0f, const float Thickness = 0.0f, const int32 Segments = 9, const float VerticalOffset = 2.0f, const bool bDepthTest = true);
+
+	/** Debug Draw a transform trajectory */
+	UFUNCTION(BlueprintCallable, Category = "Draw Debug Library", meta = (DevelopmentOnly))
+	static UE_API void DrawDebugTransformTrajectory(const FDebugDrawer& Drawer, UPARAM(ref) const FTransformTrajectory& TransformTrajectory, const FLinearColor Color = FLinearColor::Red, const float Radius = 5.0f, const float Thickness = 0.0f, const float VerticalOffset = 2.0f, const bool bDepthTest = true);
+
+	/** Draws a trajectory in the style of trajectory drawing in persona */
+	static UE_API void DrawDebugRangeTrajectoryArrayView(
+		const FDebugDrawer& Drawer,
+		const TArrayView<const FVector> TrajectoryLocations,
+		const TArrayView<const FQuat4f> TrajectoryRotations,
+		const FVector& ForwardVector,
+		const bool bDrawOrientations,
+		const bool bDrawStartingAtOrigin,
+		const FVector& OriginOffset,
+		const bool bDepthTest = true);
+
+	/** Draws a flat 2d motif representing mover position and rotation */
+	UFUNCTION(BlueprintCallable, Category = "Draw Debug Library", meta = (DevelopmentOnly, AutoCreateRefTerm = "Location, Rotation"))
+	static UE_API void DrawDebugMoverOrientation(const FDebugDrawer& Drawer, const FVector& Location, const FRotator& Rotation, const FVector ForwardVector = FVector::ForwardVector, const float Scale = 30.0f, const FLinearColor Color = FLinearColor::Red, const float Thickness = 0.0f, const bool bDepthTest = true);
+
+public: // Poses / Skeletons
+
+	/** Debug Draw a pose made up of bone locations and velocities */
+	UFUNCTION(BlueprintCallable, Category = "Draw Debug Library", meta = (DevelopmentOnly, AutoCreateRefTerm = "RelativeTransform"))
+	static UE_API void DrawDebugPose(const FDebugDrawer& Drawer, const TArray<FVector>& BoneLocations, const TArray<FVector>& BoneLinearVelocities, const FTransform& RelativeTransform = FTransform(), const FLinearColor Color = FLinearColor::Red, const float DrawVelocityLineScale = 1.0f, const float Thickness = 0.0f, const bool bDepthTest = true);
+	static UE_API void DrawDebugPoseFromArrayViews(const FDebugDrawer& Drawer, const TArrayView<const FVector> BoneLocations, const TArrayView<const FVector> BoneLinearVelocities, const FTransform& RelativeTransform = FTransform(), const FLinearColor Color = FLinearColor::Red, const float DrawVelocityLineScale = 1.0f, const float Thickness = 0.0f, const bool bDepthTest = true);
+
+	/** Draws a skeleton made up of the given bone transforms */
+	static UE_API void DrawDebugSkeletonArrayView(
+		const FDebugDrawer& Drawer,
+		const FVector RootLocation,
+		const TArrayView<const FVector> BoneLocations,
+		const TArrayView<const FQuat4f> BoneRotations,
+		const TArrayView<const FVector3f> BoneScales,
+		const TArrayView<const int32> BoneIndices,
+		const FReferenceSkeleton& ReferenceSkeleton,
+		const FLinearColor& Color,
+		const float BoneDrawSize);
+
+	/** Draws a simple skeleton made up of the given bone locations */
+	static UE_API void DrawDebugSimpleSkeletonArrayView(
+		const FDebugDrawer& Drawer,
+		const TArrayView<const FVector> BoneLocations,
+		const TArrayView<const int32> BoneIndices,
+		const TArrayView<const int32> BoneParents,
+		const FLinearColor& Color,
+		const float Thickness = 0.0f,
+		const bool bDepthTest = false);
+
+	/** Draws lines representing bone velocities */
+	static UE_API void DrawDebugBoneVelocitiesArrayView(
+		const FDebugDrawer& Drawer,
+		const TArrayView<const FVector> BoneLocations,
+		const TArrayView<const FVector3f> BoneVelocities,
+		const TArrayView<const int32> BoneIndices,
+		const FLinearColor& Color,
+		const float Scale = 100.0f,
+		const float Thickness = 0.0f,
+		const bool bDepthTest = false);
+
+	/** Draws capsules for a skeleton */
+	static UE_API void DrawDebugSkeletonCapsulesArrayView(
+		const FDebugDrawer& Drawer,
+		const TArrayView<const FVector> BoneLocations,
+		const TArrayView<const FQuat4f> BoneRotations,
+		const TArrayView<const FVector3f> BoneLocalLocations,
+		const TArrayView<const int32> BoneIndices,
+		const TArrayView<const int32> BoneParents,
+		const float MinCapsuleRadius = 1.0f,
+		const float MaxCapsuleRadius = 10.0f,
+		const FLinearColor Color = FLinearColor::Red,
+		const float Thickness = 0.0f,
+		const int32 Segments = 9,
+		const bool bDepthTest = true);
+
+public: // Text
+
+	/** Get the number of line segments required to draw debug string. */
+	static UE_API int32 DrawDebugStringViewSegmentNum(const FStringView String, const bool bMonospaced = true);
+
+	/** Get the dimensions of a draw debug string. Useful for aligning or centering text. */
+	UFUNCTION(BlueprintCallable, Category = "Draw Debug Library", meta = (DevelopmentOnly, AutoCreateRefTerm = "String"))
+	static UE_API FVector DrawDebugStringDimensions(const FString& String, const float Scale = 10.0f, const bool bMonospaced = true);
+	static UE_API FVector DrawDebugStringViewDimensions(const FStringView String, const float Scale = 10.0f, const bool bMonospaced = true);
+
+	/** Get the dimensions of a draw debug text. Useful for aligning or centering text. */
+	UFUNCTION(BlueprintCallable, Category = "Draw Debug Library", meta = (DevelopmentOnly, AutoCreateRefTerm = "Text"))
+	static UE_API FVector DrawDebugTextDimensions(const FText& Text, const float Scale = 10.0f, const bool bMonospaced = true);
+
+	/** Get the dimensions of a draw debug name. Useful for aligning or centering text. */
+	UFUNCTION(BlueprintCallable, Category = "Draw Debug Library", meta = (DevelopmentOnly, AutoCreateRefTerm = "Name"))
+	static UE_API FVector DrawDebugNameDimensions(const FName& Name, const float Scale = 10.0f, const bool bMonospaced = true);
+
+	/** Debug Draw a string. Will only render ASCII characters. */
+	UFUNCTION(BlueprintCallable, Category = "Draw Debug Library", meta = (DevelopmentOnly, AutoCreateRefTerm = "String, Location, Rotation"))
+	static UE_API void DrawDebugString(const FDebugDrawer& Drawer, const FString& String, const FVector& Location, const FRotator& Rotation, const float Scale = 10.0f, const bool bMonospaced = true, const FLinearColor Color = FLinearColor::Red, const float Thickness = 0.0f, const bool bDepthTest = true);
+	static UE_API void DrawDebugStringView(const FDebugDrawer& Drawer, const FStringView String, const FVector& Location, const FRotator& Rotation, const float Scale = 10.0f, const bool bMonospaced = true, const FLinearColor Color = FLinearColor::Red, const float Thickness = 0.0f, const bool bDepthTest = true);
+
+	/** Debug Draw a name */
+	UFUNCTION(BlueprintCallable, Category = "Draw Debug Library", meta = (DevelopmentOnly, AutoCreateRefTerm = "Name, Location, Rotation"))
+	static UE_API void DrawDebugName(const FDebugDrawer& Drawer, const FName& Name, const FVector& Location, const FRotator& Rotation, const float Scale = 10.0f, const bool bMonospaced = true, const FLinearColor Color = FLinearColor::Red, const float Thickness = 0.0f, const bool bDepthTest = true);
+
+	/** Debug Draw some text. Will only render ASCII characters. */
+	UFUNCTION(BlueprintCallable, Category = "Draw Debug Library", meta = (DevelopmentOnly, AutoCreateRefTerm = "Text, Location, Rotation"))
+	static UE_API void DrawDebugText(const FDebugDrawer& Drawer, const FText& Text, const FVector& Location, const FRotator& Rotation, const float Scale = 10.0f, const bool bMonospaced = true, const FLinearColor Color = FLinearColor::Red, const float Thickness = 0.0f, const bool bDepthTest = true);
+
+public: // Visual Logger Text
+
+	/** Debug Draw a string to the visual logger. Will do nothing if not using a VisualLoggerDrawer. Will not display on-screen during recording. */
+	UFUNCTION(BlueprintCallable, Category = "Draw Debug Library", meta = (DevelopmentOnly, AutoCreateRefTerm = "String, Location"))
+	static UE_API void DrawDebugVisualLoggerString(const FDebugDrawer& Drawer, const FString& String, const FVector& Location, const FLinearColor Color = FLinearColor::Red);
+	static UE_API void DrawDebugVisualLoggerStringView(const FDebugDrawer& Drawer, const FStringView String, const FVector& Location, const FLinearColor Color = FLinearColor::Red);
+
+	/** Debug Draw a name to the visual logger. Will do nothing if not using a VisualLoggerDrawer. Will not display on-screen during recording. */
+	UFUNCTION(BlueprintCallable, Category = "Draw Debug Library", meta = (DevelopmentOnly, AutoCreateRefTerm = "Name, Location"))
+	static UE_API void DrawDebugVisualLoggerName(const FDebugDrawer& Drawer, const FName& Name, const FVector& Location, const FLinearColor Color = FLinearColor::Red);
+
+	/** Debug Draw some text to the visual logger. Will do nothing if not using a VisualLoggerDrawer. Will not display on-screen during recording. */
+	UFUNCTION(BlueprintCallable, Category = "Draw Debug Library", meta = (DevelopmentOnly, AutoCreateRefTerm = "Text, Location"))
+	static UE_API void DrawDebugVisualLoggerText(const FDebugDrawer& Drawer, const FText& Text, const FVector& Location, const FLinearColor Color = FLinearColor::Red);
+
+
+public: // Graphs
+
+	/** Debug Draw a simple graph. */
+	UFUNCTION(BlueprintCallable, Category = "Draw Debug Library", meta = (DevelopmentOnly, AutoCreateRefTerm = "Location, Rotation, Title, XaxisLabel, YaxisLabel"))
+	static UE_API void DrawDebugGraph(
+		const FDebugDrawer& Drawer,
+		const FVector& Location,
+		const FRotator& Rotation,
+		const TArray<float>& Xvalues,
+		const TArray<float>& Yvalues,
+		const float Xmin = 0.0f,
+		const float Xmax = 1.0f,
+		const float Ymin = 0.0f,
+		const float Ymax = 1.0f,
+		const float XaxisLength = 100.0f,
+		const float YaxisLength = 100.0f,
+		const FString& Title = TEXT(""),
+		const FString& XaxisLabel = TEXT(""),
+		const FString& YaxisLabel = TEXT(""),
+		const float TextScale = 10.0f,
+		const FLinearColor Color = FLinearColor::Red,
+		const float AxisThickness = 1.0f,
+		const float LineThickness = 0.0f,
+		const float TextThickness = 0.0f,
+		const bool bDepthTest = true);
+
+	static UE_API void DrawDebugGraphView(
+		const FDebugDrawer& Drawer,
+		const FVector& Location,
+		const FRotator& Rotation,
+		const TArrayView<const float> Xvalues,
+		const TArrayView<const float> Yvalues,
+		const float Xmin = 0.0f,
+		const float Xmax = 1.0f,
+		const float Ymin = 0.0f,
+		const float Ymax = 1.0f,
+		const float XaxisLength = 100.0f,
+		const float YaxisLength = 100.0f,
+		const FStringView Title = TEXT(""),
+		const FStringView XaxisLabel = TEXT(""),
+		const FStringView YaxisLabel = TEXT(""),
+		const float TextScale = 10.0f,
+		const FLinearColor Color = FLinearColor::Red,
+		const float AxisThickness = 1.0f,
+		const float LineThickness = 0.0f,
+		const float TextThickness = 0.0f,
+		const bool bDepthTest = true);
+};
+
+#undef UE_API
+
