Refactoring and other updates to DrawDebugLibrary

diff --git a/Engine/Plugins/Experimental/DrawDebugLibrary/Source/DrawDebugLibrary/Private/DrawDebugLibrary.cpp b/Engine/Plugins/Experimental/DrawDebugLibrary/Source/DrawDebugLibrary/Private/DrawDebugLibrary.cpp
--- a/Engine/Plugins/Experimental/DrawDebugLibrary/Source/DrawDebugLibrary/Private/DrawDebugLibrary.cpp
+++ b/Engine/Plugins/Experimental/DrawDebugLibrary/Source/DrawDebugLibrary/Private/DrawDebugLibrary.cpp
@@ -9,6 +9,7 @@
 #include "DrawDebugHelpers.h"
 #include "Animation/AnimInstanceProxy.h"
 #include "VisualLogger/VisualLogger.h"
+#include "Components/SkinnedMeshComponent.h"
 
 void FDrawDebugLibraryModule::StartupModule()
 {
@@ -405,48 +406,54 @@ namespace UE::DrawDebugLibrary::Private
 	}
 }
 
-FDebugDrawer UDrawDebugLibrary::MakeNullDebugDrawer()
+FDebugDrawer::FDebugDrawer() = default;
+
+FDebugDrawer::FDebugDrawer(UObject* InObject) : FDebugDrawer()
 {
-	return FDebugDrawer();
+	World = InObject ? InObject->GetWorld() : nullptr;
 }
 
-FDebugDrawer UDrawDebugLibrary::MakeDebugDrawer(UObject* Object)
+FDebugDrawer::FDebugDrawer(UWorld* InWorld) : FDebugDrawer()
 {
-	return MakeObjectDebugDrawer(Object);
+	World = InWorld;
 }
 
-FDebugDrawer UDrawDebugLibrary::MakeObjectDebugDrawer(UObject* Object)
+FDebugDrawer::FDebugDrawer(FAnimInstanceProxy* InAnimInstanceProxy) : FDebugDrawer()
 {
-	return MakeWorldDebugDrawer(Object ? Object->GetWorld() : nullptr);
+	AnimInstanceProxy = InAnimInstanceProxy;
 }
 
-FDebugDrawer UDrawDebugLibrary::MakeWorldDebugDrawer(UWorld* World)
+FDebugDrawer::FDebugDrawer(FPrimitiveDrawInterface* InPrimitiveDrawInterface) : FDebugDrawer()
 {
-	FDebugDrawer Drawer;
-	Drawer.World = World;
-	return Drawer;
+	PDI = InPrimitiveDrawInterface;
 }
 
-FDebugDrawer UDrawDebugLibrary::MakeAnimInstanceProxyDebugDrawer(FAnimInstanceProxy* AnimInstanceProxy)
+FDebugDrawer FDebugDrawer::MakeDebugDrawer()
 {
-	FDebugDrawer Drawer;
-	Drawer.AnimInstanceProxy = AnimInstanceProxy;
-	return Drawer;
+	return FDebugDrawer();
 }
 
-FDebugDrawer UDrawDebugLibrary::MakePDIDebugDrawer(FPrimitiveDrawInterface* PrimitiveDrawInterface)
+FDebugDrawer FDebugDrawer::MakeDebugDrawer(UObject* Object)
 {
-	FDebugDrawer Drawer;
-	Drawer.PDI = PrimitiveDrawInterface;
-	return Drawer;
+	return MakeDebugDrawer(Object ? Object->GetWorld() : nullptr);
 }
 
-FDebugDrawer UDrawDebugLibrary::MakeVisualLoggerDebugDrawer(UObject* Object, const FName Category, const EDrawDebugLogVerbosity Verbosity, const bool bDrawToScene)
+FDebugDrawer FDebugDrawer::MakeDebugDrawer(UWorld* World)
 {
-	return MakeVisualLoggerDebugDrawerFromObject(Object, Category, Verbosity, bDrawToScene);
+	return FDebugDrawer(World);
 }
 
-FDebugDrawer UDrawDebugLibrary::MakeVisualLoggerDebugDrawerFromObject(UObject* Object, const FName Category, const EDrawDebugLogVerbosity Verbosity, const bool bDrawToScene)
+FDebugDrawer FDebugDrawer::MakeDebugDrawer(FAnimInstanceProxy* AnimInstanceProxy)
+{
+	return FDebugDrawer(AnimInstanceProxy);
+}
+
+FDebugDrawer FDebugDrawer::MakeDebugDrawer(FPrimitiveDrawInterface* PrimitiveDrawInterface)
+{
+	return FDebugDrawer(PrimitiveDrawInterface);
+}
+
+FDebugDrawer FDebugDrawer::MakeVisualLoggerDebugDrawer(UObject* Object, const FName Category, const EDrawDebugLogVerbosity Verbosity, const bool bDrawToScene)
 {
 	FDebugDrawer Drawer;
 	Drawer.World = (Object && bDrawToScene) ? Object->GetWorld() : nullptr;
@@ -456,7 +463,7 @@ FDebugDrawer UDrawDebugLibrary::MakeVisualLoggerDebugDrawerFromObject(UObject* O
 	return Drawer;
 }
 
-FDebugDrawer UDrawDebugLibrary::MakeVisualLoggerDebugDrawerFromObjectWithCategory(UObject* Object, const FLogCategoryBase& Category, const ELogVerbosity::Type Verbosity, const bool bDrawToScene)
+FDebugDrawer FDebugDrawer::MakeVisualLoggerDebugDrawer(UObject* Object, const FLogCategoryBase& Category, const ELogVerbosity::Type Verbosity, const bool bDrawToScene)
 {
 	FDebugDrawer Drawer;
 	Drawer.World = (Object && bDrawToScene) ? Object->GetWorld() : nullptr;
@@ -466,12 +473,7 @@ FDebugDrawer UDrawDebugLibrary::MakeVisualLoggerDebugDrawerFromObjectWithCategor
 	return Drawer;
 }
 
-FDebugDrawer UDrawDebugLibrary::MakeMergedDebugDrawer(const TArray<FDebugDrawer>& DebugDrawers)
-{
-	return MakeMergedDebugDrawerArrayView(DebugDrawers);
-}
-
-FDebugDrawer UDrawDebugLibrary::MakeMergedDebugDrawerArrayView(const TArrayView<const FDebugDrawer> DebugDrawers)
+FDebugDrawer FDebugDrawer::MakeMergedDebugDrawer(const TArrayView<const FDebugDrawer> DebugDrawers)
 {
 	FDebugDrawer Drawer;
 
@@ -491,270 +493,167 @@ FDebugDrawer UDrawDebugLibrary::MakeMergedDebugDrawerArrayView(const TArrayView<
 	return Drawer;
 }
 
-void UDrawDebugLibrary::LogVisualLoggerString(const FDebugDrawer& Drawer, const FString& String, const FLinearColor Color)
+void FDebugDrawer::DrawPoints(const TArrayView<const FVector> Locations, const FLinearColor& Color, const float Thickness, const bool bDepthTest) const
 {
-	LogVisualLoggerStringView(Drawer, String, Color);
-}
+	const int32 Num = Locations.Num();
 
-void UDrawDebugLibrary::LogVisualLoggerStringView(const FDebugDrawer& Drawer, const FStringView String, const FLinearColor Color)
-{
-#if ENABLE_VISUAL_LOG
-	if (Drawer.VisualLoggerObject)
+	if (PDI)
 	{
-		FVisualLogEntry* CurrentEntry = nullptr;
-		if (FVisualLogger::CheckVisualLogInputInternal(Drawer.VisualLoggerObject, Drawer.VisualLoggerCategory, Drawer.VisualLoggerVerbosity, /*OutWorld=*/ nullptr, &CurrentEntry))
+		for (int32 Idx = 0; Idx < Num; Idx++)
 		{
-			CurrentEntry->LogLines.Add(FVisualLogLine(Drawer.VisualLoggerCategory, Drawer.VisualLoggerVerbosity, FString(String)));
+			PDI->DrawPoint(Locations[Idx], Color, Thickness, bDepthTest ? SDPG_World : SDPG_Foreground);
 		}
 	}
-#endif
-}
-
-void UDrawDebugLibrary::LogVisualLoggerText(const FDebugDrawer& Drawer, const FText& Text, const FLinearColor Color)
-{
-	LogVisualLoggerStringView(Drawer, Text.ToString(), Color);
-}
-
-void UDrawDebugLibrary::MakeLinearlySpacedFloatArray(TArray<float>& OutValues, const float Start, const float Stop, const int32 Num)
-{
-	OutValues.SetNumUninitialized(Num);
-
-	if (Num == 1)
-	{
-		OutValues[0] = Start;
-		return;
-	}
-
-	for (int32 Idx = 0; Idx < Num; Idx++)
-	{
-		OutValues[Idx] = Start + (Stop - Start) * ((float)(Idx) / (Num - 1));
-	}
-}
-
-void UDrawDebugLibrary::AddToFloatHistoryArray(TArray<float>& InOutValues, const float NewValue, const int32 MaxHistoryNum)
-{
-	if (MaxHistoryNum <= 0)
-	{
-		InOutValues.Empty();
-		return;
-	}
-
-	InOutValues.Add(NewValue);
 
-	while (InOutValues.Num() > MaxHistoryNum)
+#if UE_ENABLE_DEBUG_DRAWING
+	if (World)
 	{
-		// Shift Values
-		for (int32 Idx = 0; Idx < InOutValues.Num() - 1; Idx++)
+		for (int32 Idx = 0; Idx < Num; Idx++)
 		{
-			InOutValues[Idx] = InOutValues[Idx + 1];
+			::DrawDebugPoint(World, Locations[Idx], Thickness, Color.ToFColor(true), false, -1.0f, bDepthTest ? 0 : 1);
 		}
-
-		// Pop End
-		InOutValues.Pop();
 	}
-}
 
-void UDrawDebugLibrary::AddToVectorHistoryArray(TArray<FVector>& InOutValues, const FVector NewValue, const int32 MaxHistoryNum)
-{
-	if (MaxHistoryNum <= 0)
+	if (AnimInstanceProxy)
 	{
-		InOutValues.Empty();
-		return;
-	}
-
-	InOutValues.Add(NewValue);
-
-	while (InOutValues.Num() > MaxHistoryNum)
-	{
-		// Shift Values
-		for (int32 Idx = 0; Idx < InOutValues.Num() - 1; Idx++)
+		for (int32 Idx = 0; Idx < Num; Idx++)
 		{
-			InOutValues[Idx] = InOutValues[Idx + 1];
+			AnimInstanceProxy->AnimDrawDebugPoint(Locations[Idx], Thickness, Color.ToFColor(true), false, -1.0f, bDepthTest ? SDPG_World : SDPG_Foreground);
 		}
-
-		// Pop End
-		InOutValues.Pop();
-	}
-}
-
-void UDrawDebugLibrary::DrawDebugPoint(const FDebugDrawer& Drawer, const FVector& Location, const FLinearColor Color, const float Size, const bool bDepthTest)
-{
-	if (Drawer.PDI)
-	{
-		Drawer.PDI->DrawPoint(Location, Color, Size, bDepthTest ? SDPG_World : SDPG_Foreground);
-	}
-
-#if UE_ENABLE_DEBUG_DRAWING
-	if (Drawer.World)
-	{
-		::DrawDebugPoint(Drawer.World, Location, Size, Color.ToFColor(true), false, -1.0f, bDepthTest ? 0 : 1);
-	}
-
-	if (Drawer.AnimInstanceProxy)
-	{
-		Drawer.AnimInstanceProxy->AnimDrawDebugPoint(Location, Size, Color.ToFColor(true), false, -1.0f, bDepthTest ? SDPG_World : SDPG_Foreground);
 	}
 #endif
 
 #if ENABLE_VISUAL_LOG
-	if (Drawer.VisualLoggerObject)
+	if (VisualLoggerObject)
 	{
 		FVisualLogEntry* CurrentEntry = nullptr;
-		if (FVisualLogger::CheckVisualLogInputInternal(Drawer.VisualLoggerObject, Drawer.VisualLoggerCategory, Drawer.VisualLoggerVerbosity, /*OutWorld=*/ nullptr, &CurrentEntry))
+		if (FVisualLogger::CheckVisualLogInputInternal(VisualLoggerObject, VisualLoggerCategory, VisualLoggerVerbosity, /*OutWorld=*/ nullptr, &CurrentEntry))
 		{
-			const uint16 IntegerSize = Size == 0.0f ? 0 : FMath::CeilToInt(FMath::Max(2.0f * Size, 0.0f));
-			FVisualLogShapeElement Element(TEXT(""), Color.ToFColor(false), IntegerSize, Drawer.VisualLoggerCategory);
+			// The uin16 integer size is difficult to match to the float thickness used by the other debug drawers. This is my best estimation but I
+			// am not 100% sure the mapping is correct.
+			const uint16 IntegerSize = Thickness == 0.0f ? 0 : FMath::CeilToInt(FMath::Max(2.0f * Thickness, 0.0f));
+			FVisualLogShapeElement Element(TEXT(""), Color.ToFColor(false), IntegerSize, VisualLoggerCategory);
 			Element.Type = EVisualLoggerShapeElement::SinglePoint;
 			Element.Radius = 0.0f;
-			Element.Verbosity = Drawer.VisualLoggerVerbosity;
-			Element.Points.Add(Location);
+			Element.Verbosity = VisualLoggerVerbosity;
+			Element.Points.Append(Locations);
 			CurrentEntry->ElementsToDraw.Add(Element);
 		}
 	}
 #endif
 }
 
-void UDrawDebugLibrary::DrawDebugPoints(const FDebugDrawer& Drawer, const TArray<FVector>& Locations, const FLinearColor Color, const float Size, const bool bDepthTest)
+void FDebugDrawer::DrawLines(const TArrayView<const FVector> StartLocations, const TArrayView<const FVector> EndLocations, const FLinearColor& Color, const float Thickness, const bool bDepthTest) const
 {
-	DrawDebugPointsArrayView(Drawer, Locations, Color, Size, bDepthTest);
-}
-
-void UDrawDebugLibrary::DrawDebugPointsArrayView(const FDebugDrawer& Drawer, const TArrayView<const FVector> Locations, const FLinearColor Color, const float Size, const bool bDepthTest)
-{
-	const int32 Num = Locations.Num();
+	const int32 Num = FMath::Min(StartLocations.Num(), EndLocations.Num());
 
-	if (Drawer.PDI)
+	if (PDI)
 	{
 		for (int32 Idx = 0; Idx < Num; Idx++)
 		{
-			Drawer.PDI->DrawPoint(Locations[Idx], Color, Size, bDepthTest ? SDPG_World : SDPG_Foreground);
+			if (bDepthTest)
+			{
+				PDI->DrawTranslucentLine(StartLocations[Idx], EndLocations[Idx], Color, SDPG_World, FMath::Max(Thickness, 0.0f), 0.0f, Thickness <= 0.0f);
+			}
+			else
+			{
+				PDI->DrawTranslucentLine(StartLocations[Idx], EndLocations[Idx], Color, Thickness == 0.0f ? SDPG_Foreground : SDPG_World, Thickness, Thickness == 0.0f ? 1.0f : UE_MAX_FLT);
+			}
 		}
 	}
 
 #if UE_ENABLE_DEBUG_DRAWING
-	if (Drawer.World)
+	if (World)
 	{
 		for (int32 Idx = 0; Idx < Num; Idx++)
 		{
-			::DrawDebugPoint(Drawer.World, Locations[Idx], Size, Color.ToFColor(true), false, -1.0f, bDepthTest ? 0 : 1);
+			::DrawDebugLine(World, StartLocations[Idx], EndLocations[Idx], Color.ToFColor(true), false, -1.0f, bDepthTest ? 0 : 1, Thickness);
 		}
 	}
 
-	if (Drawer.AnimInstanceProxy)
+	if (AnimInstanceProxy)
 	{
 		for (int32 Idx = 0; Idx < Num; Idx++)
 		{
-			Drawer.AnimInstanceProxy->AnimDrawDebugPoint(Locations[Idx], Size, Color.ToFColor(true), false, -1.0f, bDepthTest ? SDPG_World : SDPG_Foreground);
+			AnimInstanceProxy->AnimDrawDebugLine(StartLocations[Idx], EndLocations[Idx], Color.ToFColor(true), false, -1.0f, Thickness, bDepthTest ? SDPG_World : SDPG_Foreground);
 		}
 	}
 #endif
 
 #if ENABLE_VISUAL_LOG
-	if (Drawer.VisualLoggerObject)
+	if (VisualLoggerObject)
 	{
 		FVisualLogEntry* CurrentEntry = nullptr;
-		if (FVisualLogger::CheckVisualLogInputInternal(Drawer.VisualLoggerObject, Drawer.VisualLoggerCategory, Drawer.VisualLoggerVerbosity, /*OutWorld=*/ nullptr, &CurrentEntry))
-		{
-			const uint16 IntegerSize = Size == 0.0f ? 0 : FMath::CeilToInt(FMath::Max(2.0f * Size, 0.0f));
-			FVisualLogShapeElement Element(TEXT(""), Color.ToFColor(false), IntegerSize, Drawer.VisualLoggerCategory);
-			Element.Type = EVisualLoggerShapeElement::SinglePoint;
-			Element.Radius = 0.0f;
-			Element.Verbosity = Drawer.VisualLoggerVerbosity;
-			Element.Points.Append(Locations);
-			CurrentEntry->ElementsToDraw.Add(Element);
-		}
-	}
-#endif
-}
-
-void UDrawDebugLibrary::DrawDebugLine(const FDebugDrawer& Drawer, const FVector& StartLocation, const FVector& EndLocation, const FLinearColor Color, const float Thickness, const bool bDepthTest)
-{
-	if (Drawer.PDI)
-	{
-		Drawer.PDI->DrawTranslucentLine(StartLocation, EndLocation, Color, SDPG_World, Thickness, bDepthTest ? 0.0f : UE_MAX_FLT);
-	}
-	
-#if UE_ENABLE_DEBUG_DRAWING
-	if (Drawer.World)
-	{
-		::DrawDebugLine(Drawer.World, StartLocation, EndLocation, Color.ToFColor(true), false, -1.0f, bDepthTest ? 0 : 1, Thickness);
-	}
-
-	if (Drawer.AnimInstanceProxy)
-	{
-		Drawer.AnimInstanceProxy->AnimDrawDebugLine(StartLocation, EndLocation, Color.ToFColor(true), false, -1.0f, Thickness, bDepthTest ? SDPG_World : SDPG_Foreground);
-	}
-#endif
-
-#if ENABLE_VISUAL_LOG
-	if (Drawer.VisualLoggerObject)
-	{
-		FVisualLogEntry* CurrentEntry = nullptr;
-		if (FVisualLogger::CheckVisualLogInputInternal(Drawer.VisualLoggerObject, Drawer.VisualLoggerCategory, Drawer.VisualLoggerVerbosity, /*OutWorld=*/ nullptr, &CurrentEntry))
+		if (FVisualLogger::CheckVisualLogInputInternal(VisualLoggerObject, VisualLoggerCategory, VisualLoggerVerbosity, /*OutWorld=*/ nullptr, &CurrentEntry))
 		{
 			const uint16 IntegerThickness = Thickness == 0.0f ? 0 : FMath::CeilToInt(FMath::Max(2.0f * Thickness, 0.0f));
-			FVisualLogShapeElement Element(TEXT(""), Color.ToFColor(false), IntegerThickness, Drawer.VisualLoggerCategory);
+			FVisualLogShapeElement Element(TEXT(""), Color.ToFColor(false), IntegerThickness, VisualLoggerCategory);
 			Element.Type = EVisualLoggerShapeElement::Segment;
-			Element.Verbosity = Drawer.VisualLoggerVerbosity;
-			Element.Points.Reserve(2);
-			Element.Points.Add(StartLocation);
-			Element.Points.Add(EndLocation);
+			Element.Verbosity = VisualLoggerVerbosity;
+			Element.Points.Reserve(2 * Num);
+			for (int32 Idx = 0; Idx < Num; Idx++)
+			{
+				Element.Points.Add(StartLocations[Idx]);
+				Element.Points.Add(EndLocations[Idx]);
+			}
 			CurrentEntry->ElementsToDraw.Add(Element);
 		}
 	}
 #endif
 }
 
-void UDrawDebugLibrary::DrawDebugLines(const FDebugDrawer& Drawer, const TArray<FVector>& StartLocations, const TArray<FVector>& EndLocations, const FLinearColor Color, const float Thickness, const bool bDepthTest)
+void FDebugDrawer::DrawLinesPairs(const TArrayView<const TPair<FVector, FVector>> Segments, const FLinearColor& Color, const float Thickness, const bool bDepthTest) const
 {
-	DrawDebugLinesArrayView(Drawer, StartLocations, EndLocations, Color, Thickness, bDepthTest);
-}
-
-void UDrawDebugLibrary::DrawDebugLinesArrayView(const FDebugDrawer& Drawer, const TArrayView<const FVector> StartLocations, const TArrayView<const FVector> EndLocations, const FLinearColor Color, const float Thickness, const bool bDepthTest)
-{
-	const int32 Num = FMath::Min(StartLocations.Num(), EndLocations.Num());
+	const int32 Num = Segments.Num();
 
-	if (Drawer.PDI)
+	if (PDI)
 	{
 		for (int32 Idx = 0; Idx < Num; Idx++)
 		{
-			Drawer.PDI->DrawTranslucentLine(StartLocations[Idx], EndLocations[Idx], Color, SDPG_World, Thickness, bDepthTest ? 0.0f : UE_MAX_FLT);
+			if (bDepthTest)
+			{
+				PDI->DrawTranslucentLine(Segments[Idx].Key, Segments[Idx].Value, Color, SDPG_World, FMath::Max(Thickness, 0.0f), 0.0f, Thickness <= 0.0f);
+			}
+			else
+			{
+				PDI->DrawTranslucentLine(Segments[Idx].Key, Segments[Idx].Value, Color, Thickness == 0.0f ? SDPG_Foreground : SDPG_World, Thickness, Thickness == 0.0f ? 1.0f : UE_MAX_FLT);
+			}
 		}
 	}
 
 #if UE_ENABLE_DEBUG_DRAWING
-	if (Drawer.World)
+	if (World)
 	{
 		for (int32 Idx = 0; Idx < Num; Idx++)
 		{
-			::DrawDebugLine(Drawer.World, StartLocations[Idx], EndLocations[Idx], Color.ToFColor(true), false, -1.0f, bDepthTest ? 0 : 1, Thickness);
+			::DrawDebugLine(World, Segments[Idx].Key, Segments[Idx].Value, Color.ToFColor(true), false, -1.0f, bDepthTest ? 0 : 1, Thickness);
 		}
 	}
 
-	if (Drawer.AnimInstanceProxy)
+	if (AnimInstanceProxy)
 	{
 		for (int32 Idx = 0; Idx < Num; Idx++)
 		{
-			Drawer.AnimInstanceProxy->AnimDrawDebugLine(StartLocations[Idx], EndLocations[Idx], Color.ToFColor(true), false, -1.0f, Thickness, bDepthTest ? SDPG_World : SDPG_Foreground);
+			AnimInstanceProxy->AnimDrawDebugLine(Segments[Idx].Key, Segments[Idx].Value, Color.ToFColor(true), false, -1.0f, Thickness, bDepthTest ? SDPG_World : SDPG_Foreground);
 		}
 	}
 #endif
 
 #if ENABLE_VISUAL_LOG
-	if (Drawer.VisualLoggerObject)
+	if (VisualLoggerObject)
 	{
 		FVisualLogEntry* CurrentEntry = nullptr;
-		if (FVisualLogger::CheckVisualLogInputInternal(Drawer.VisualLoggerObject, Drawer.VisualLoggerCategory, Drawer.VisualLoggerVerbosity, /*OutWorld=*/ nullptr, &CurrentEntry))
+		if (FVisualLogger::CheckVisualLogInputInternal(VisualLoggerObject, VisualLoggerCategory, VisualLoggerVerbosity, /*OutWorld=*/ nullptr, &CurrentEntry))
 		{
 			const uint16 IntegerThickness = Thickness == 0.0f ? 0 : FMath::CeilToInt(FMath::Max(2.0f * Thickness, 0.0f));
-			FVisualLogShapeElement Element(TEXT(""), Color.ToFColor(false), IntegerThickness, Drawer.VisualLoggerCategory);
+			FVisualLogShapeElement Element(TEXT(""), Color.ToFColor(false), IntegerThickness, VisualLoggerCategory);
 			Element.Type = EVisualLoggerShapeElement::Segment;
-			Element.Verbosity = Drawer.VisualLoggerVerbosity;
+			Element.Verbosity = VisualLoggerVerbosity;
 			Element.Points.Reserve(2 * Num);
 			for (int32 Idx = 0; Idx < Num; Idx++)
 			{
-				Element.Points.Add(StartLocations[Idx]);
-				Element.Points.Add(EndLocations[Idx]);
+				Element.Points.Add(Segments[Idx].Key);
+				Element.Points.Add(Segments[Idx].Value);
 			}
 			CurrentEntry->ElementsToDraw.Add(Element);
 		}
@@ -762,59 +661,206 @@ void UDrawDebugLibrary::DrawDebugLinesArrayView(const FDebugDrawer& Drawer, cons
 #endif
 }
 
-void UDrawDebugLibrary::DrawDebugLinesPairsArrayView(const FDebugDrawer& Drawer, const TArrayView<const TPair<FVector, FVector>> Segments, const FLinearColor Color, const float Thickness, const bool bDepthTest)
+void FDebugDrawer::VisualLoggerLogString(const FStringView String, const FLinearColor& Color) const
 {
-	const int32 Num = Segments.Num();
-
-	if (Drawer.PDI)
+#if ENABLE_VISUAL_LOG
+	if (VisualLoggerObject)
 	{
-		for (int32 Idx = 0; Idx < Num; Idx++)
+		FVisualLogEntry* CurrentEntry = nullptr;
+		if (FVisualLogger::CheckVisualLogInputInternal(VisualLoggerObject, VisualLoggerCategory, VisualLoggerVerbosity, /*OutWorld=*/ nullptr, &CurrentEntry))
 		{
-			Drawer.PDI->DrawTranslucentLine(Segments[Idx].Key, Segments[Idx].Value, Color, SDPG_World, Thickness, bDepthTest ? 0.0f : UE_MAX_FLT);
+			CurrentEntry->LogLines.Add(FVisualLogLine(VisualLoggerCategory, VisualLoggerVerbosity, FString(String)));
 		}
 	}
+#endif
+}
 
-#if UE_ENABLE_DEBUG_DRAWING
-	if (Drawer.World)
+void FDebugDrawer::VisualLoggerDrawString(const FStringView String, const FVector& Location, const FLinearColor& Color) const
+{
+#if ENABLE_VISUAL_LOG
+	if (VisualLoggerObject)
 	{
-		for (int32 Idx = 0; Idx < Num; Idx++)
+		FVisualLogEntry* CurrentEntry = nullptr;
+		if (FVisualLogger::CheckVisualLogInputInternal(VisualLoggerObject, VisualLoggerCategory, VisualLoggerVerbosity, /*OutWorld=*/ nullptr, &CurrentEntry))
 		{
-			::DrawDebugLine(Drawer.World, Segments[Idx].Key, Segments[Idx].Value, Color.ToFColor(true), false, -1.0f, bDepthTest ? 0 : 1, Thickness);
+			FVisualLogShapeElement Element(FString(String), Color.ToFColor(false), 0, VisualLoggerCategory);
+			Element.Type = EVisualLoggerShapeElement::SinglePoint;
+			Element.Verbosity = VisualLoggerVerbosity;
+			Element.Points.Add(Location);
+			CurrentEntry->ElementsToDraw.Add(Element);
 		}
 	}
+#endif
+}
+
+FDebugDrawer UDrawDebugLibrary::MakeNullDebugDrawer()
+{
+	return FDebugDrawer::MakeDebugDrawer();
+}
+
+FDebugDrawer UDrawDebugLibrary::MakeDebugDrawer(UObject* Object)
+{
+	return FDebugDrawer::MakeDebugDrawer(Object);
+}
+
+FDebugDrawer UDrawDebugLibrary::MakeObjectDebugDrawer(UObject* Object)
+{
+	return FDebugDrawer::MakeDebugDrawer(Object);
+}
+
+FDebugDrawer UDrawDebugLibrary::MakeWorldDebugDrawer(UWorld* World)
+{
+	return FDebugDrawer::MakeDebugDrawer(World);
+}
+
+FDebugDrawer UDrawDebugLibrary::MakeAnimInstanceProxyDebugDrawer(FAnimInstanceProxy* AnimInstanceProxy)
+{
+	return FDebugDrawer::MakeDebugDrawer(AnimInstanceProxy);
+}
+
+FDebugDrawer UDrawDebugLibrary::MakePDIDebugDrawer(FPrimitiveDrawInterface* PrimitiveDrawInterface)
+{
+	return FDebugDrawer::MakeDebugDrawer(PrimitiveDrawInterface);
+}
+
+FDebugDrawer UDrawDebugLibrary::MakeVisualLoggerDebugDrawer(UObject* Object, const FName Category, const EDrawDebugLogVerbosity Verbosity, const bool bDrawToScene)
+{
+	return FDebugDrawer::MakeVisualLoggerDebugDrawer(Object, Category, Verbosity, bDrawToScene);
+}
+
+FDebugDrawer UDrawDebugLibrary::MakeVisualLoggerDebugDrawerFromObject(UObject* Object, const FName Category, const EDrawDebugLogVerbosity Verbosity, const bool bDrawToScene)
+{
+	return FDebugDrawer::MakeVisualLoggerDebugDrawer(Object, Category, Verbosity, bDrawToScene);
+}
+
+FDebugDrawer UDrawDebugLibrary::MakeVisualLoggerDebugDrawerFromObjectWithCategory(UObject* Object, const FLogCategoryBase& Category, const ELogVerbosity::Type Verbosity, const bool bDrawToScene)
+{
+	return FDebugDrawer::MakeVisualLoggerDebugDrawer(Object, Category, Verbosity, bDrawToScene);
+}
+
+FDebugDrawer UDrawDebugLibrary::MakeMergedDebugDrawer(const TArray<FDebugDrawer>& DebugDrawers)
+{
+	return FDebugDrawer::MakeMergedDebugDrawer(DebugDrawers);
+}
+
+FDebugDrawer UDrawDebugLibrary::MakeMergedDebugDrawerArrayView(const TArrayView<const FDebugDrawer> DebugDrawers)
+{
+	return FDebugDrawer::MakeMergedDebugDrawer(DebugDrawers);
+}
+
+void UDrawDebugLibrary::VisualLoggerLogString(const FDebugDrawer& Drawer, const FString& String, const FLinearColor Color)
+{
+	Drawer.VisualLoggerLogString(String, Color);
+}
+
+void UDrawDebugLibrary::VisualLoggerLogStringView(const FDebugDrawer& Drawer, const FStringView String, const FLinearColor Color)
+{
+	Drawer.VisualLoggerLogString(String, Color);
+}
+
+void UDrawDebugLibrary::VisualLoggerLogText(const FDebugDrawer& Drawer, const FText& Text, const FLinearColor Color)
+{
+	Drawer.VisualLoggerLogString(Text.ToString(), Color);
+}
+
+void UDrawDebugLibrary::MakeLinearlySpacedFloatArray(TArray<float>& OutValues, const float Start, const float Stop, const int32 Num)
+{
+	OutValues.SetNumUninitialized(Num);
+
+	if (Num == 1)
+	{
+		OutValues[0] = Start;
+		return;
+	}
 
-	if (Drawer.AnimInstanceProxy)
+	for (int32 Idx = 0; Idx < Num; Idx++)
 	{
-		for (int32 Idx = 0; Idx < Num; Idx++)
+		OutValues[Idx] = Start + (Stop - Start) * ((float)(Idx) / (Num - 1));
+	}
+}
+
+void UDrawDebugLibrary::AddToFloatHistoryArray(TArray<float>& InOutValues, const float NewValue, const int32 MaxHistoryNum)
+{
+	if (MaxHistoryNum <= 0)
+	{
+		InOutValues.Empty();
+		return;
+	}
+
+	InOutValues.Add(NewValue);
+
+	while (InOutValues.Num() > MaxHistoryNum)
+	{
+		// Shift Values
+		for (int32 Idx = 0; Idx < InOutValues.Num() - 1; Idx++)
 		{
-			Drawer.AnimInstanceProxy->AnimDrawDebugLine(Segments[Idx].Key, Segments[Idx].Value, Color.ToFColor(true), false, -1.0f, Thickness, bDepthTest ? SDPG_World : SDPG_Foreground);
+			InOutValues[Idx] = InOutValues[Idx + 1];
 		}
+
+		// Pop End
+		InOutValues.Pop();
 	}
-#endif
+}
 
-#if ENABLE_VISUAL_LOG
-	if (Drawer.VisualLoggerObject)
+void UDrawDebugLibrary::AddToVectorHistoryArray(TArray<FVector>& InOutValues, const FVector NewValue, const int32 MaxHistoryNum)
+{
+	if (MaxHistoryNum <= 0)
 	{
-		FVisualLogEntry* CurrentEntry = nullptr;
-		if (FVisualLogger::CheckVisualLogInputInternal(Drawer.VisualLoggerObject, Drawer.VisualLoggerCategory, Drawer.VisualLoggerVerbosity, /*OutWorld=*/ nullptr, &CurrentEntry))
+		InOutValues.Empty();
+		return;
+	}
+
+	InOutValues.Add(NewValue);
+
+	while (InOutValues.Num() > MaxHistoryNum)
+	{
+		// Shift Values
+		for (int32 Idx = 0; Idx < InOutValues.Num() - 1; Idx++)
 		{
-			const uint16 IntegerThickness = Thickness == 0.0f ? 0 : FMath::CeilToInt(FMath::Max(2.0f * Thickness, 0.0f));
-			FVisualLogShapeElement Element(TEXT(""), Color.ToFColor(false), IntegerThickness, Drawer.VisualLoggerCategory);
-			Element.Type = EVisualLoggerShapeElement::Segment;
-			Element.Verbosity = Drawer.VisualLoggerVerbosity;
-			Element.Points.Reserve(2 * Num);
-			for (int32 Idx = 0; Idx < Num; Idx++)
-			{
-				Element.Points.Add(Segments[Idx].Key);
-				Element.Points.Add(Segments[Idx].Value);
-			}
-			CurrentEntry->ElementsToDraw.Add(Element);
+			InOutValues[Idx] = InOutValues[Idx + 1];
 		}
+
+		// Pop End
+		InOutValues.Pop();
 	}
-#endif
 }
 
-void UDrawDebugLibrary::DrawDebugTriangularBasePyramid(const FDebugDrawer& Drawer, const FVector& Location, const FRotator& Rotation, const FLinearColor Color, const float Length, const float Width, const float Thickness, const bool bDepthTest)
+void UDrawDebugLibrary::DrawDebugPoint(const FDebugDrawer& Drawer, const FVector& Location, const FLinearColor Color, const float Thickness, const bool bDepthTest)
+{
+	Drawer.DrawPoints({ Location }, Color, Thickness, bDepthTest);
+}
+
+void UDrawDebugLibrary::DrawDebugPoints(const FDebugDrawer& Drawer, const TArray<FVector>& Locations, const FLinearColor Color, const float Thickness, const bool bDepthTest)
+{
+	Drawer.DrawPoints(Locations, Color, Thickness, bDepthTest);
+}
+
+void UDrawDebugLibrary::DrawDebugPointsArrayView(const FDebugDrawer& Drawer, const TArrayView<const FVector> Locations, const FLinearColor Color, const float Thickness, const bool bDepthTest)
+{
+	Drawer.DrawPoints(Locations, Color, Thickness, bDepthTest);
+}
+
+void UDrawDebugLibrary::DrawDebugLine(const FDebugDrawer& Drawer, const FVector& StartLocation, const FVector& EndLocation, const FLinearColor Color, const float Thickness, const bool bDepthTest)
+{
+	Drawer.DrawLines({ StartLocation }, { EndLocation }, Color, Thickness, bDepthTest);
+}
+
+void UDrawDebugLibrary::DrawDebugLines(const FDebugDrawer& Drawer, const TArray<FVector>& StartLocations, const TArray<FVector>& EndLocations, const FLinearColor Color, const float Thickness, const bool bDepthTest)
+{
+	Drawer.DrawLines(StartLocations, EndLocations, Color, Thickness, bDepthTest);
+}
+
+void UDrawDebugLibrary::DrawDebugLinesArrayView(const FDebugDrawer& Drawer, const TArrayView<const FVector> StartLocations, const TArrayView<const FVector> EndLocations, const FLinearColor Color, const float Thickness, const bool bDepthTest)
+{
+	Drawer.DrawLines(StartLocations, EndLocations, Color, Thickness, bDepthTest);
+}
+
+void UDrawDebugLibrary::DrawDebugLinesPairsArrayView(const FDebugDrawer& Drawer, const TArrayView<const TPair<FVector, FVector>> Segments, const FLinearColor Color, const float Thickness, const bool bDepthTest)
+{
+	Drawer.DrawLinesPairs(Segments, Color, Thickness, bDepthTest);
+}
+
+void UDrawDebugLibrary::DrawDebugTriangularBasePyramid(const FDebugDrawer& Drawer, const FVector& Location, const FRotator& Rotation, const FLinearColor Color, const float Thickness, const bool bDepthTest, const float Length, const float Width)
 {
 	const FVector V0 = Location + Rotation.RotateVector(FVector(0, 0, Length/2));
 	const FVector V1 = Location + Rotation.RotateVector(FVector(-Width / 4 * FMath::Sqrt(3.0f), -Width / 2, -Length / 2));
@@ -831,7 +877,7 @@ void UDrawDebugLibrary::DrawDebugTriangularBasePyramid(const FDebugDrawer& Drawe
 	}, Color, Thickness, bDepthTest);
 }
 
-void UDrawDebugLibrary::DrawDebugSquareBasePyramid(const FDebugDrawer& Drawer, const FVector& Location, const FRotator& Rotation, const FLinearColor Color, const float Length, const float Width, const float Thickness, const bool bDepthTest)
+void UDrawDebugLibrary::DrawDebugSquareBasePyramid(const FDebugDrawer& Drawer, const FVector& Location, const FRotator& Rotation, const FLinearColor Color, const float Thickness, const bool bDepthTest, const float Length, const float Width)
 {
 	const FVector V0 = Location + Rotation.RotateVector(FVector(0, 0, Length / 2));
 	const FVector V1 = Location + Rotation.RotateVector(FVector(-Width / 2, -Width / 2, -Length / 2));
@@ -851,7 +897,7 @@ void UDrawDebugLibrary::DrawDebugSquareBasePyramid(const FDebugDrawer& Drawer, c
 		}, Color, Thickness, bDepthTest);
 }
 
-void UDrawDebugLibrary::DrawDebugCone(const FDebugDrawer& Drawer, const FVector& Location, const FRotator& Rotation, const FLinearColor Color, const float Length, const float Radius, const int32 Segments, const float Thickness, const bool bDepthTest)
+void UDrawDebugLibrary::DrawDebugCone(const FDebugDrawer& Drawer, const FVector& Location, const FRotator& Rotation, const FLinearColor Color, const float Thickness, const bool bDepthTest, const float Length, const float Radius, const int32 Segments)
 {
 	const int32 Num = FMath::Max(Segments, 2);
 
@@ -871,9 +917,8 @@ void UDrawDebugLibrary::DrawDebugCone(const FDebugDrawer& Drawer, const FVector&
 	DrawDebugLinesPairsArrayView(Drawer, LineSegments, Color, Thickness, bDepthTest);
 }
 
-void UDrawDebugLibrary::DrawDebugArc(const FDebugDrawer& Drawer, const FVector& Location, const FRotator& Rotation, const float Angle, const FLinearColor Color, const float Radius, const float Thickness, const int32 Segments, const bool bDepthTest)
+void UDrawDebugLibrary::DrawDebugArc(const FDebugDrawer& Drawer, const FVector& Location, const FRotator& Rotation, const float Angle, const FLinearColor Color, const float Thickness, const bool bDepthTest, const float Radius, const int32 Segments)
 {
-
 	TArray<TPair<FVector, FVector>, TInlineAllocator<17>> LineSegments;
 	LineSegments.Reserve(Segments);
 
@@ -889,18 +934,107 @@ void UDrawDebugLibrary::DrawDebugArc(const FDebugDrawer& Drawer, const FVector&
 	DrawDebugLinesPairsArrayView(Drawer, LineSegments, Color, Thickness, bDepthTest);
 }
 
-void UDrawDebugLibrary::DrawDebugCircle(const FDebugDrawer& Drawer, const FVector& Location, const FRotator& Rotation, const FLinearColor Color, const float Radius, const float Thickness, const int32 Segments, const bool bDepthTest)
+void UDrawDebugLibrary::DrawDebugCircle(const FDebugDrawer& Drawer, const FVector& Location, const FRotator& Rotation, const FLinearColor Color, const float Thickness, const bool bDepthTest, const float Radius, const int32 Segments)
+{
+	DrawDebugArc(Drawer, Location, Rotation, 360.0f, Color, Thickness, bDepthTest, Radius, FMath::Max(Segments, 3));
+}
+
+void UDrawDebugLibrary::DrawDebugCircleTick(const FDebugDrawer& Drawer, const FVector& Location, const FRotator& Rotation, const float Angle, const FLinearColor Color, const float Thickness, const bool bDepthTest, const float Radius, const float Length, bool bInside)
 {
-	DrawDebugArc(Drawer, Location, Rotation, 360.0f, Color, Radius, Thickness, Segments, bDepthTest);
+	DrawDebugCircleTicksArrayView(Drawer, Location, Rotation, { Angle }, Color, Thickness, bDepthTest, Radius, Length, bInside);
 }
 
-void UDrawDebugLibrary::DrawDebugCircleOutline(const FDebugDrawer& Drawer, const FVector& Location, const FRotator& Rotation, const FLinearColor Color, const float InnerRadius, const float OuterRadius, const float Thickness, const int32 Segments, const bool bDepthTest)
+void UDrawDebugLibrary::DrawDebugCircleTicks(const FDebugDrawer& Drawer, const FVector& Location, const FRotator& Rotation, const TArray<float>& Angles, const FLinearColor Color, const float Thickness, const bool bDepthTest, const float Radius, const float Length, bool bInside)
 {
-	DrawDebugCircle(Drawer, Location, Rotation, Color, InnerRadius, Thickness, Segments, bDepthTest);
-	DrawDebugCircle(Drawer, Location, Rotation, Color, OuterRadius, Thickness, Segments, bDepthTest);
+	DrawDebugCircleTicksArrayView(Drawer, Location, Rotation, Angles, Color, Thickness, bDepthTest, Radius, Length, bInside);
 }
 
-void UDrawDebugLibrary::DrawDebugLocator(const FDebugDrawer& Drawer, const FVector& Location, const FRotator& Rotation, const FLinearColor Color, const float DrawRadius, const float Thickness, const bool bDepthTest)
+void UDrawDebugLibrary::DrawDebugCircleTicksArrayView(const FDebugDrawer& Drawer, const FVector& Location, const FRotator& Rotation, const TArrayView<const float> Angles, const FLinearColor Color, const float Thickness, const bool bDepthTest, const float Radius, const float Length, bool bInside)
+{
+	const int32 TickNum = Angles.Num();
+
+	TArray<TPair<FVector, FVector>, TInlineAllocator<17>> LineSegments;
+	LineSegments.Reserve(TickNum);
+
+	for (int32 TickIdx = 0; TickIdx < TickNum; TickIdx++)
+	{
+		const float Angle = Angles[TickIdx];
+		const float Offset = bInside ? -Length : +Length;
+		const FVector Start = Location + Rotation.RotateVector(Radius * FVector(FMath::Cos(FMath::DegreesToRadians(Angle)), FMath::Sin(FMath::DegreesToRadians(Angle)), 0.0f));
+		const FVector End = Location + Rotation.RotateVector((Radius + Offset) * FVector(FMath::Cos(FMath::DegreesToRadians(Angle)), FMath::Sin(FMath::DegreesToRadians(Angle)), 0.0f));
+		LineSegments.Add({ Start, End });
+	}
+
+	DrawDebugLinesPairsArrayView(Drawer, LineSegments, Color, Thickness, bDepthTest);
+}
+
+void UDrawDebugLibrary::DrawDebugTriangle(const FDebugDrawer& Drawer, const FVector& Location, const FRotator& Rotation, const FLinearColor Color, const float Thickness, const bool bDepthTest, const float Radius)
+{
+	DrawDebugCircle(Drawer, Location, Rotation, Color, Thickness, bDepthTest, Radius, 3);
+}
+
+void UDrawDebugLibrary::DrawDebugSquare(const FDebugDrawer& Drawer, const FVector& Location, const FRotator& Rotation, const FLinearColor Color, const float Thickness, const bool bDepthTest, const float HalfLength)
+{
+	const FVector V0 = Location + Rotation.RotateVector(HalfLength * FVector(+1, +1, 0));
+	const FVector V1 = Location + Rotation.RotateVector(HalfLength * FVector(-1, +1, 0));
+	const FVector V2 = Location + Rotation.RotateVector(HalfLength * FVector(-1, -1, 0));
+	const FVector V3 = Location + Rotation.RotateVector(HalfLength * FVector(+1, -1, 0));
+
+	DrawDebugLinesPairsArrayView(Drawer, {
+		{ V0, V1 },
+		{ V1, V2 },
+		{ V2, V3 },
+		{ V3, V0 },
+		}, Color, Thickness, bDepthTest);
+}
+
+void UDrawDebugLibrary::DrawDebugDiamond(const FDebugDrawer& Drawer, const FVector& Location, const FRotator& Rotation, const FLinearColor Color, const float Thickness, const bool bDepthTest, const float Radius)
+{
+	DrawDebugCircle(Drawer, Location, Rotation, Color, Thickness, bDepthTest, Radius, 4);
+}
+
+void UDrawDebugLibrary::DrawDebugPentagon(const FDebugDrawer& Drawer, const FVector& Location, const FRotator& Rotation, const FLinearColor Color, const float Thickness, const bool bDepthTest, const float Radius)
+{
+	DrawDebugCircle(Drawer, Location, Rotation, Color, Thickness, bDepthTest, Radius, 5);
+}
+
+void UDrawDebugLibrary::DrawDebugHexagon(const FDebugDrawer& Drawer, const FVector& Location, const FRotator& Rotation, const FLinearColor Color, const float Thickness, const bool bDepthTest, const float Radius)
+{
+	DrawDebugCircle(Drawer, Location, Rotation, Color, Thickness, bDepthTest, Radius, 6);
+}
+
+void UDrawDebugLibrary::DrawDebugHeptagon(const FDebugDrawer& Drawer, const FVector& Location, const FRotator& Rotation, const FLinearColor Color, const float Thickness, const bool bDepthTest, const float Radius)
+{
+	DrawDebugCircle(Drawer, Location, Rotation, Color, Thickness, bDepthTest, Radius, 7);
+}
+
+void UDrawDebugLibrary::DrawDebugOctagon(const FDebugDrawer& Drawer, const FVector& Location, const FRotator& Rotation, const FLinearColor Color, const float Thickness, const bool bDepthTest, const float Radius)
+{
+	DrawDebugCircle(Drawer, Location, Rotation, Color, Thickness, bDepthTest, Radius, 8);
+}
+
+void UDrawDebugLibrary::DrawDebugNonagon(const FDebugDrawer& Drawer, const FVector& Location, const FRotator& Rotation, const FLinearColor Color, const float Thickness, const bool bDepthTest, const float Radius)
+{
+	DrawDebugCircle(Drawer, Location, Rotation, Color,Thickness, bDepthTest, Radius, 9);
+}
+
+void UDrawDebugLibrary::DrawDebugDecagon(const FDebugDrawer& Drawer, const FVector& Location, const FRotator& Rotation, const FLinearColor Color, const float Thickness, const bool bDepthTest, const float Radius)
+{
+	DrawDebugCircle(Drawer, Location, Rotation, Color, Thickness, bDepthTest, Radius, 10);
+}
+
+void UDrawDebugLibrary::DrawDebugRegularPolygon(const FDebugDrawer& Drawer, const FVector& Location, const FRotator& Rotation, const FLinearColor Color, const float Thickness, const bool bDepthTest, const float Radius, const int32 Sides)
+{
+	DrawDebugCircle(Drawer, Location, Rotation, Color, Thickness, bDepthTest, Radius, Sides);
+}
+
+void UDrawDebugLibrary::DrawDebugCircleOutline(const FDebugDrawer& Drawer, const FVector& Location, const FRotator& Rotation, const FLinearColor Color, const float Thickness, const bool bDepthTest, const float InnerRadius, const float OuterRadius, const int32 Segments)
+{
+	DrawDebugCircle(Drawer, Location, Rotation, Color, Thickness, bDepthTest, InnerRadius, Segments);
+	DrawDebugCircle(Drawer, Location, Rotation, Color, Thickness, bDepthTest, OuterRadius, Segments);
+}
+
+void UDrawDebugLibrary::DrawDebugLocator(const FDebugDrawer& Drawer, const FVector& Location, const FRotator& Rotation, const FLinearColor Color, const float Thickness, const bool bDepthTest, const float DrawRadius)
 {
 
 	DrawDebugLinesPairsArrayView(Drawer, {
@@ -910,7 +1044,7 @@ void UDrawDebugLibrary::DrawDebugLocator(const FDebugDrawer& Drawer, const FVect
 		}, Color, Thickness, bDepthTest);
 }
 
-void UDrawDebugLibrary::DrawDebugCrossLocator(const FDebugDrawer& Drawer, const FVector& Location, const FRotator& Rotation, const FLinearColor Color, const float DrawRadius, const float Thickness, const bool bDepthTest)
+void UDrawDebugLibrary::DrawDebugCrossLocator(const FDebugDrawer& Drawer, const FVector& Location, const FRotator& Rotation, const FLinearColor Color, const float Thickness, const bool bDepthTest, const float DrawRadius)
 {
 	DrawDebugLinesPairsArrayView(Drawer, {
 		{ Location + Rotation.RotateVector(FVector(+DrawRadius, +DrawRadius, -DrawRadius)), Location + Rotation.RotateVector(FVector(-DrawRadius, -DrawRadius, +DrawRadius)) },
@@ -920,7 +1054,7 @@ void UDrawDebugLibrary::DrawDebugCrossLocator(const FDebugDrawer& Drawer, const
 		}, Color, Thickness, bDepthTest);
 }
 
-void UDrawDebugLibrary::DrawDebugBox(const FDebugDrawer& Drawer, const FVector& Location, const FRotator& Rotation, const FVector HalfExtents, const FLinearColor Color, const float Thickness, const bool bDepthTest)
+void UDrawDebugLibrary::DrawDebugBox(const FDebugDrawer& Drawer, const FVector& Location, const FRotator& Rotation, const FLinearColor Color, const float Thickness, const bool bDepthTest, const FVector HalfExtents)
 {
 	const FVector C0 = Location + Rotation.RotateVector(FVector(+1, +1, -1) * HalfExtents);
 	const FVector C1 = Location + Rotation.RotateVector(FVector(+1, -1, -1) * HalfExtents);
@@ -954,10 +1088,10 @@ void UDrawDebugLibrary::DrawDebugSphere(
 	const FVector& Location,
 	const FRotator& Rotation,
 	const FLinearColor Color,
-	const float Radius,
 	const float Thickness,
-	const int32 Segments,
-	const bool bDepthTest)
+	const bool bDepthTest,
+	const float Radius,
+	const int32 Segments)
 {
 	const int32 SegmentNum = (FMath::Max(Segments, 4) / 2) * 2;
 
@@ -984,10 +1118,32 @@ void UDrawDebugLibrary::DrawDebugSphere(
 	}
 
 	DrawDebugLinesPairsArrayView(Drawer, LineSegments, Color, Thickness, bDepthTest);
+}
+
+void UDrawDebugLibrary::DrawDebugSimpleSphere(const FDebugDrawer& Drawer, const FVector& Location, const FRotator& Rotation, const FLinearColor Color, const float Thickness, const bool bDepthTest, const float Radius, const int32 Segments)
+{
+	TArray<TPair<FVector, FVector>, TInlineAllocator<39>> LineSegments;
+	LineSegments.Reserve(Segments*3);
 
+	for (int32 SegmentIdx = 0; SegmentIdx < Segments; SegmentIdx++)
+	{
+		const float Angle0 = (((float)SegmentIdx + 0) / Segments) * UE_TWO_PI;
+		const float Angle1 = (((float)SegmentIdx + 1) / Segments) * UE_TWO_PI;
+		const FVector X0 = Location + Rotation.RotateVector(Radius * FVector(0.0f, FMath::Cos(Angle0), FMath::Sin(Angle0)));
+		const FVector X1 = Location + Rotation.RotateVector(Radius * FVector(0.0f, FMath::Cos(Angle1), FMath::Sin(Angle1)));
+		const FVector Y0 = Location + Rotation.RotateVector(Radius * FVector(FMath::Cos(Angle0), 0.0f, FMath::Sin(Angle0)));
+		const FVector Y1 = Location + Rotation.RotateVector(Radius * FVector(FMath::Cos(Angle1), 0.0f, FMath::Sin(Angle1)));
+		const FVector Z0 = Location + Rotation.RotateVector(Radius * FVector(FMath::Cos(Angle0), FMath::Sin(Angle0), 0.0f));
+		const FVector Z1 = Location + Rotation.RotateVector(Radius * FVector(FMath::Cos(Angle1), FMath::Sin(Angle1), 0.0f));
+		LineSegments.Add({ X0, X1 });
+		LineSegments.Add({ Y0, Y1 });
+		LineSegments.Add({ Z0, Z1 });
+	}
+
+	DrawDebugLinesPairsArrayView(Drawer, LineSegments, Color, Thickness, bDepthTest);
 }
 
-void UDrawDebugLibrary::DrawDebugCapsule(const FDebugDrawer& Drawer, const FVector& Location, const FRotator& Rotation, const FLinearColor Color, const float Radius, const float HalfLength, const float Thickness, const int32 Segments, const bool bDepthTest)
+void UDrawDebugLibrary::DrawDebugCapsule(const FDebugDrawer& Drawer, const FVector& Location, const FRotator& Rotation, const FLinearColor Color, const float Thickness, const bool bDepthTest, const float Radius, const float HalfLength, const int32 Segments)
 {
 	const int32 SegmentNum = (FMath::Max(Segments, 4) / 4) * 4;
 
@@ -1037,6 +1193,13 @@ void UDrawDebugLibrary::DrawDebugCapsule(const FDebugDrawer& Drawer, const FVect
 	DrawDebugLinesPairsArrayView(Drawer, LineSegments, Color, Thickness, bDepthTest);
 }
 
+void UDrawDebugLibrary::DrawDebugCapsuleLine(const FDebugDrawer& Drawer, const FVector& StartLocation, const FVector& EndLocation, const FLinearColor Color, const float Thickness, const bool bDepthTest, const float Radius, const int32 Segments)
+{
+	const FQuat Rotation = UE::DrawDebugLibrary::Private::FindOrientation(StartLocation, EndLocation);
+	const float Length = (EndLocation - StartLocation).Length();
+
+	DrawDebugCapsule(Drawer, StartLocation + Rotation.RotateVector(FVector(Length / 2.0f, 0.0f, 0.0f)), (Rotation * FQuat::MakeFromRotationVector(FVector(0.0f, UE_HALF_PI, 0.0f))).Rotator(), Color, Thickness, bDepthTest, Radius, Length / 2.0f, Segments);
+}
 
 void UDrawDebugLibrary::DrawDebugFrustum(const FDebugDrawer& Drawer, const FMatrix& FrustumToWorld, const FLinearColor Color, const float Thickness, const bool bDepthTest)
 {
@@ -1081,89 +1244,155 @@ void UDrawDebugLibrary::DrawDebugArrow(
 	const FVector& StartLocation,
 	const FVector& EndLocation,
 	const FLinearColor Color,
-	const float ArrowHeadSize,
 	const float Thickness,
-	const bool bDepthTest)
+	const bool bDepthTest,
+	const FDrawDebugArrowSettings& Settings)
 {
 	const FQuat Rotation = UE::DrawDebugLibrary::Private::FindOrientation(StartLocation, EndLocation);
 
-	const float ArrowSqrt = FMath::Sqrt(ArrowHeadSize);
-
-	DrawDebugLinesPairsArrayView(Drawer, {
-		{ StartLocation, EndLocation },
-		{ EndLocation, EndLocation + Rotation.RotateVector(FVector(-ArrowSqrt, +ArrowSqrt, 0)) },
-		{ EndLocation, EndLocation + Rotation.RotateVector(FVector(-ArrowSqrt, -ArrowSqrt, 0)) },
-		}, Color, Thickness, bDepthTest);
+	return DrawDebugOrientedArrow(Drawer, StartLocation, Rotation.Rotator(), (EndLocation - StartLocation).Length(), Color, Thickness, bDepthTest, Settings);
 }
 
-void UDrawDebugLibrary::DrawDebugOrientedArrow(const FDebugDrawer& Drawer, const FVector& Location, const FRotator& Rotation, const float Length, const FLinearColor Color, const float ArrowHeadSize, const float Thickness, const bool bDepthTest)
+void UDrawDebugLibrary::DrawDebugOrientedArrow(
+	const FDebugDrawer& Drawer, 
+	const FVector& Location, 
+	const FRotator& Rotation, 
+	const float Length,
+	const FLinearColor Color,
+	const float Thickness,
+	const bool bDepthTest,
+	const FDrawDebugArrowSettings& Settings)
 {
-	const float ArrowSqrt = FMath::Sqrt(ArrowHeadSize);
+	const FVector EndLocation = Location + Rotation.RotateVector(FVector(Length, 0.0f, 0.0f));
 
-	DrawDebugLinesPairsArrayView(Drawer, {
-		{ Location, Location + Rotation.RotateVector(FVector(Length, 0.0f, 0.0f)) },
-		{ Location + Rotation.RotateVector(FVector(Length, 0.0f, 0.0f)), Location + Rotation.RotateVector(FVector(Length - ArrowSqrt, +ArrowSqrt, 0.0f)) },
-		{ Location + Rotation.RotateVector(FVector(Length, 0.0f, 0.0f)), Location + Rotation.RotateVector(FVector(Length - ArrowSqrt, -ArrowSqrt, 0.0f)) },
-		}, Color, Thickness, bDepthTest);
-}
+	DrawDebugLine(Drawer, Location, EndLocation, Color, Thickness, bDepthTest);
 
-void UDrawDebugLibrary::DrawDebugFourWayArrow(const FDebugDrawer& Drawer, const FVector& StartLocation, const FVector& EndLocation, const FLinearColor Color, const float ArrowHeadSize, const float Thickness, const bool bDepthTest)
-{
-	const FQuat Rotation = UE::DrawDebugLibrary::Private::FindOrientation(StartLocation, EndLocation);
+	if (Settings.bArrowheadOnStart)
+	{
+		switch (Settings.ArrowHeadStartType)
+		{
+		case EDrawDebugArrowHead::Simple:
+		{
+			DrawDebugLinesPairsArrayView(Drawer, {
+				{ Location, Location + Rotation.RotateVector(FVector(+Settings.ArrowHeadStartWidth, +Settings.ArrowHeadStartWidth, 0.0f)) },
+				{ Location, Location + Rotation.RotateVector(FVector(+Settings.ArrowHeadStartWidth, -Settings.ArrowHeadStartWidth, 0.0f)) },
+				}, Color, Thickness, bDepthTest);
+			break;
+		}
+
+		case EDrawDebugArrowHead::FourWay:
+		{
+			DrawDebugLinesPairsArrayView(Drawer, {
+				{ Location, Location + Rotation.RotateVector(FVector(+Settings.ArrowHeadStartWidth, +Settings.ArrowHeadStartWidth, 0)) },
+				{ Location, Location + Rotation.RotateVector(FVector(+Settings.ArrowHeadStartWidth, -Settings.ArrowHeadStartWidth, 0)) },
+				{ Location, Location + Rotation.RotateVector(FVector(+Settings.ArrowHeadStartWidth, 0, +Settings.ArrowHeadStartWidth)) },
+				{ Location, Location + Rotation.RotateVector(FVector(+Settings.ArrowHeadStartWidth, 0, -Settings.ArrowHeadStartWidth)) },
+				}, Color, Thickness, bDepthTest);
+			break;
+		}
 
-	const float ArrowSqrt = FMath::Sqrt(ArrowHeadSize);
+		case EDrawDebugArrowHead::Triangular:
+		{
+			DrawDebugTriangularBasePyramid(Drawer,
+				Location + Rotation.RotateVector(FVector(+Settings.ArrowHeadStartWidth / 2, 0, 0)),
+				(Rotation.Quaternion() * FQuat::MakeFromRotationVector(FVector(0, -UE_PI / 2, 0))).Rotator(),
+				Color, Thickness, bDepthTest, Settings.ArrowHeadStartWidth, Settings.ArrowHeadStartWidth);
+			break;
+		}
 
-	DrawDebugLinesPairsArrayView(Drawer, {
-		{ StartLocation, EndLocation },
-		{ EndLocation, EndLocation + Rotation.RotateVector(FVector(-ArrowSqrt, +ArrowSqrt, 0)) },
-		{ EndLocation, EndLocation + Rotation.RotateVector(FVector(-ArrowSqrt, -ArrowSqrt, 0)) },
-		{ EndLocation, EndLocation + Rotation.RotateVector(FVector(-ArrowSqrt, 0, +ArrowSqrt)) },
-		{ EndLocation, EndLocation + Rotation.RotateVector(FVector(-ArrowSqrt, 0, -ArrowSqrt)) },
-		}, Color, Thickness, bDepthTest);
-}
+		case EDrawDebugArrowHead::Square:
+		{
+			DrawDebugSquareBasePyramid(Drawer,
+				Location + Rotation.RotateVector(FVector(+Settings.ArrowHeadStartWidth / 2, 0, 0)),
+				(Rotation.Quaternion() * FQuat::MakeFromRotationVector(FVector(0, -UE_PI / 2, 0))).Rotator(),
+				Color, Thickness, bDepthTest, Settings.ArrowHeadStartWidth, Settings.ArrowHeadStartWidth);
+			break;
+		}
 
-void UDrawDebugLibrary::DrawDebugGroundTargetArrow(const FDebugDrawer& Drawer, const FVector& Location, const FLinearColor Color, const float Length, const float ArrowHeadSize, const float Radius, const int32 Segements, const float Thickness, const bool bDepthTest)
-{
-	DrawDebugFourWayArrow(
-		Drawer,
-		Location + FVector(0.0f, 0.0f, Length),
-		Location,
-		Color,
-		ArrowHeadSize,
-		Thickness,
-		bDepthTest);
+		case EDrawDebugArrowHead::Cone:
+		{
+			DrawDebugCone(Drawer,
+				Location + Rotation.RotateVector(FVector(+Settings.ArrowHeadStartWidth / 2, 0, 0)),
+				(Rotation.Quaternion() * FQuat::MakeFromRotationVector(FVector(0, -UE_PI / 2, 0))).Rotator(),
+				Color, Thickness, bDepthTest, Settings.ArrowHeadStartWidth, Settings.ArrowHeadStartWidth / 2, 7);
+			break;
+		}
 
-	DrawDebugCircle(Drawer, Location, FRotator::ZeroRotator, Color, Radius, Thickness, Segements, bDepthTest);
-}
+		default: checkNoEntry();
+		}
+	}
 
-void UDrawDebugLibrary::DrawTriangularArrow(const FDebugDrawer& Drawer, const FVector& StartLocation, const FVector& EndLocation, const FLinearColor Color, const float ArrowHeadSize, const float Thickness, const bool bDepthTest)
-{
-	const FQuat Rotation = UE::DrawDebugLibrary::Private::FindOrientation(StartLocation, EndLocation);
+	if (Settings.bArrowheadOnEnd)
+	{
+		switch (Settings.ArrowHeadEndType)
+		{
+		case EDrawDebugArrowHead::Simple:
+		{
+			DrawDebugLinesPairsArrayView(Drawer, {
+				{ EndLocation, EndLocation + Rotation.RotateVector(FVector(-Settings.ArrowHeadEndWidth, +Settings.ArrowHeadEndWidth, 0.0f)) },
+				{ EndLocation, EndLocation + Rotation.RotateVector(FVector(-Settings.ArrowHeadEndWidth, -Settings.ArrowHeadEndWidth, 0.0f)) },
+				}, Color, Thickness, bDepthTest);
+			break;
+		}
 
-	const float ArrowSqrt = FMath::Sqrt(ArrowHeadSize);
+		case EDrawDebugArrowHead::FourWay:
+		{
+			DrawDebugLinesPairsArrayView(Drawer, {
+				{ EndLocation, EndLocation + Rotation.RotateVector(FVector(-Settings.ArrowHeadEndWidth, +Settings.ArrowHeadEndWidth, 0)) },
+				{ EndLocation, EndLocation + Rotation.RotateVector(FVector(-Settings.ArrowHeadEndWidth, -Settings.ArrowHeadEndWidth, 0)) },
+				{ EndLocation, EndLocation + Rotation.RotateVector(FVector(-Settings.ArrowHeadEndWidth, 0, +Settings.ArrowHeadEndWidth)) },
+				{ EndLocation, EndLocation + Rotation.RotateVector(FVector(-Settings.ArrowHeadEndWidth, 0, -Settings.ArrowHeadEndWidth)) },
+				}, Color, Thickness, bDepthTest);
+			break;
+		}
 
-	DrawDebugLine(Drawer, StartLocation, EndLocation, Color, Thickness, bDepthTest);
-	DrawDebugTriangularBasePyramid(Drawer, EndLocation + Rotation.RotateVector(FVector(-ArrowSqrt / 2, 0, 0)), (Rotation * FQuat::MakeFromRotationVector(FVector(0, UE_PI / 2, 0))).Rotator(), Color, ArrowSqrt, ArrowSqrt, Thickness, bDepthTest);
-}
+		case EDrawDebugArrowHead::Triangular:
+		{
+			DrawDebugTriangularBasePyramid(Drawer,
+				EndLocation + Rotation.RotateVector(FVector(-Settings.ArrowHeadEndWidth / 2, 0, 0)),
+				(Rotation.Quaternion() * FQuat::MakeFromRotationVector(FVector(0, +UE_PI / 2, 0))).Rotator(),
+				Color, Thickness, bDepthTest, Settings.ArrowHeadEndWidth, Settings.ArrowHeadEndWidth);
+			break;
+		}
 
-void UDrawDebugLibrary::DrawSquareArrow(const FDebugDrawer& Drawer, const FVector& StartLocation, const FVector& EndLocation, const FLinearColor Color, const float ArrowHeadSize, const float Thickness, const bool bDepthTest)
-{
-	const FQuat Rotation = UE::DrawDebugLibrary::Private::FindOrientation(StartLocation, EndLocation);
+		case EDrawDebugArrowHead::Square:
+		{
+			DrawDebugSquareBasePyramid(Drawer,
+				EndLocation + Rotation.RotateVector(FVector(-Settings.ArrowHeadEndWidth / 2, 0, 0)),
+				(Rotation.Quaternion() * FQuat::MakeFromRotationVector(FVector(0, +UE_PI / 2, 0))).Rotator(),
+				Color, Thickness, bDepthTest, Settings.ArrowHeadEndWidth, Settings.ArrowHeadEndWidth);
+			break;
+		}
 
-	const float ArrowSqrt = FMath::Sqrt(ArrowHeadSize);
+		case EDrawDebugArrowHead::Cone:
+		{
+			DrawDebugCone(Drawer,
+				EndLocation + Rotation.RotateVector(FVector(-Settings.ArrowHeadEndWidth / 2, 0, 0)),
+				(Rotation.Quaternion() * FQuat::MakeFromRotationVector(FVector(0, +UE_PI / 2, 0))).Rotator(),
+				Color, Thickness, bDepthTest, Settings.ArrowHeadEndWidth, Settings.ArrowHeadEndWidth / 2, 7);
+			break;
+		}
 
-	DrawDebugLine(Drawer, StartLocation, EndLocation, Color, Thickness, bDepthTest);
-	DrawDebugSquareBasePyramid(Drawer, EndLocation + Rotation.RotateVector(FVector(-ArrowSqrt / 2, 0, 0)), (Rotation * FQuat::MakeFromRotationVector(FVector(0, UE_PI / 2, 0))).Rotator(), Color, ArrowSqrt, ArrowSqrt, Thickness, bDepthTest);
+		default: checkNoEntry();
+		}
+	}
 }
 
-void UDrawDebugLibrary::DrawConeArrow(const FDebugDrawer& Drawer, const FVector& StartLocation, const FVector& EndLocation, const FLinearColor Color, const float ArrowHeadSize, const float Thickness, const int32 Segments, const bool bDepthTest)
+void UDrawDebugLibrary::DrawDebugGroundTargetArrow(const FDebugDrawer& Drawer, const FVector& Location, const FLinearColor Color, const float Thickness, const bool bDepthTest, const float Length, const float ArrowHeadSize, const float Radius, const int32 Segements)
 {
-	const FQuat Rotation = UE::DrawDebugLibrary::Private::FindOrientation(StartLocation, EndLocation);
+	FDrawDebugArrowSettings Settings;
+	Settings.ArrowHeadEndType = EDrawDebugArrowHead::FourWay;
 
-	const float ArrowSqrt = FMath::Sqrt(ArrowHeadSize);
+	DrawDebugArrow(
+		Drawer,
+		Location + FVector(0.0f, 0.0f, Length),
+		Location,
+		Color,
+		Thickness,
+		bDepthTest,
+		Settings);
 
-	DrawDebugLine(Drawer, StartLocation, EndLocation, Color, Thickness, bDepthTest);
-	DrawDebugCone(Drawer, EndLocation + Rotation.RotateVector(FVector(-ArrowSqrt / 2, 0, 0)), (Rotation * FQuat::MakeFromRotationVector(FVector(0, UE_PI / 2, 0))).Rotator(), Color, ArrowSqrt, ArrowSqrt / 2, Segments, Thickness, bDepthTest);
+	DrawDebugCircle(Drawer, Location, FRotator::ZeroRotator, Color, Thickness, bDepthTest, Radius, Segements);
 }
 
 void UDrawDebugLibrary::DrawDebugFlatArrow(
@@ -1171,32 +1400,38 @@ void UDrawDebugLibrary::DrawDebugFlatArrow(
 	const FVector& Location,
 	const FRotator& Rotation,
 	const FLinearColor Color,
-	const float Length,
-	const float Width,
 	const float Thickness,
-	const bool bDepthTest)
+	const bool bDepthTest,
+	const float Length,
+	const float Width)
 {
-	FVector ArrowPoints[7];
-	ArrowPoints[0] = Location - Rotation.RotateVector(FVector(0, Width / 2.0f, 0));
-	ArrowPoints[1] = Location + Rotation.RotateVector(FVector(0, Width / 2.0f, 0));
-	ArrowPoints[2] = ArrowPoints[0] + Rotation.RotateVector(FVector(Length / 3.0f, 0, 0));
-	ArrowPoints[3] = ArrowPoints[1] + Rotation.RotateVector(FVector(Length / 3.0f, 0, 0));
-	ArrowPoints[4] = ArrowPoints[2] - Rotation.RotateVector(FVector(0, Width / 2.0f, 0));
-	ArrowPoints[5] = ArrowPoints[3] + Rotation.RotateVector(FVector(0, Width / 2.0f, 0));
-	ArrowPoints[6] = Location + Rotation.RotateVector(FVector(Length, 0, 0));
+	const FVector V0 = Location - Rotation.RotateVector(FVector(0, Width / 4.0f, 0));
+	const FVector V1 = Location + Rotation.RotateVector(FVector(0, Width / 4.0f, 0));
+	const FVector V2 = V0 + Rotation.RotateVector(FVector(Length / 3.0f, 0, 0));
+	const FVector V3 = V1 + Rotation.RotateVector(FVector(Length / 3.0f, 0, 0));
+	const FVector V4 = V2 - Rotation.RotateVector(FVector(0, Width / 4.0f, 0));
+	const FVector V5 = V3 + Rotation.RotateVector(FVector(0, Width / 4.0f, 0));
+	const FVector V6 = Location + Rotation.RotateVector(FVector(Length, 0, 0));
 
 	DrawDebugLinesPairsArrayView(Drawer, {
-		{ ArrowPoints[0], ArrowPoints[1] },
-		{ ArrowPoints[0], ArrowPoints[2] },
-		{ ArrowPoints[1], ArrowPoints[3] },
-		{ ArrowPoints[2], ArrowPoints[4] },
-		{ ArrowPoints[3], ArrowPoints[5] },
-		{ ArrowPoints[4], ArrowPoints[6] },
-		{ ArrowPoints[5], ArrowPoints[6] },
+		{ V0, V1 },
+		{ V0, V2 },
+		{ V1, V3 },
+		{ V2, V4 },
+		{ V3, V5 },
+		{ V4, V6 },
+		{ V5, V6 },
 		}, Color, Thickness, bDepthTest);
 }
 
-void UDrawDebugLibrary::DrawDebugCatmullRomSplineStart(const FDebugDrawer& Drawer, const FVector& V0, const FVector& V1, const FVector& V2, const bool bMonotonic, const FLinearColor Color, const float Thickness, const int32 Segments, const bool bDepthTest)
+void UDrawDebugLibrary::DrawDebugCircleArrow(const FDebugDrawer& Drawer, const FVector& Location, const FRotator& Rotation, const float Angle, const FLinearColor Color, const float Thickness, const bool bDepthTest, const float Radius, const float Length, const FDrawDebugArrowSettings& ArrowSettings)
+{
+	const FVector Start = Location + Rotation.RotateVector(Radius * FVector(FMath::Cos(FMath::DegreesToRadians(Angle)), FMath::Sin(FMath::DegreesToRadians(Angle)), 0.0f));
+	const FVector End = Location + Rotation.RotateVector((Radius + Length) * FVector(FMath::Cos(FMath::DegreesToRadians(Angle)), FMath::Sin(FMath::DegreesToRadians(Angle)), 0.0f));
+	DrawDebugArrow(Drawer, Start, End, Color, Thickness, bDepthTest, ArrowSettings);
+}
+
+void UDrawDebugLibrary::DrawDebugCatmullRomSplineStart(const FDebugDrawer& Drawer, const FVector& V0, const FVector& V1, const FVector& V2, const FLinearColor Color, const float Thickness, const bool bDepthTest, const bool bMonotonic, const int32 Segments)
 {
 	const int32 SegmentNum = FMath::Max(Segments, 2);
 
@@ -1229,7 +1464,7 @@ void UDrawDebugLibrary::DrawDebugCatmullRomSplineStart(const FDebugDrawer& Drawe
 	DrawDebugLinesPairsArrayView(Drawer, LineSegments, Color, Thickness, bDepthTest);
 }
 
-void UDrawDebugLibrary::DrawDebugCatmullRomSplineEnd(const FDebugDrawer& Drawer, const FVector& V0, const FVector& V1, const FVector& V2, const bool bMonotonic, const FLinearColor Color, const float Thickness, const int32 Segments, const bool bDepthTest)
+void UDrawDebugLibrary::DrawDebugCatmullRomSplineEnd(const FDebugDrawer& Drawer, const FVector& V0, const FVector& V1, const FVector& V2, const FLinearColor Color, const float Thickness, const bool bDepthTest, const bool bMonotonic, const int32 Segments)
 {
 	const int32 SegmentNum = FMath::Max(Segments, 2);
 
@@ -1262,7 +1497,7 @@ void UDrawDebugLibrary::DrawDebugCatmullRomSplineEnd(const FDebugDrawer& Drawer,
 	DrawDebugLinesPairsArrayView(Drawer, LineSegments, Color, Thickness, bDepthTest);
 }
 
-void UDrawDebugLibrary::DrawDebugCatmullRomSplineSection(const FDebugDrawer& Drawer, const FVector& V0, const FVector& V1, const FVector& V2, const FVector& V3, const bool bMonotonic, const FLinearColor Color, const float Thickness, const int32 Segments, const bool bDepthTest)
+void UDrawDebugLibrary::DrawDebugCatmullRomSplineSection(const FDebugDrawer& Drawer, const FVector& V0, const FVector& V1, const FVector& V2, const FVector& V3, const FLinearColor Color, const float Thickness, const bool bDepthTest, const bool bMonotonic, const int32 Segments)
 {
 	const int32 SegmentNum = FMath::Max(Segments, 2);
 
@@ -1296,12 +1531,12 @@ void UDrawDebugLibrary::DrawDebugCatmullRomSplineSection(const FDebugDrawer& Dra
 }
 
 
-void UDrawDebugLibrary::DrawDebugCatmullRomSpline(const FDebugDrawer& Drawer, const TArray<FVector>& Points, const bool bMonotonic, const FLinearColor Color, const float Thickness, const int32 Segments, const bool bDepthTest)
+void UDrawDebugLibrary::DrawDebugCatmullRomSpline(const FDebugDrawer& Drawer, const TArray<FVector>& Points, const FLinearColor Color, const float Thickness, const bool bDepthTest, const bool bMonotonic, const int32 Segments)
 {
-	DrawDebugCatmullRomSplineArrayView(Drawer, Points, bMonotonic, Color, Thickness, Segments, bDepthTest);
+	DrawDebugCatmullRomSplineArrayView(Drawer, Points, Color, Thickness, bDepthTest, bMonotonic, Segments);
 }
 
-void UDrawDebugLibrary::DrawDebugCatmullRomSplineArrayView(const FDebugDrawer& Drawer, const TArrayView<const FVector> Points, const bool bMonotonic, const FLinearColor Color, const float Thickness, const int32 Segments, const bool bDepthTest)
+void UDrawDebugLibrary::DrawDebugCatmullRomSplineArrayView(const FDebugDrawer& Drawer, const TArrayView<const FVector> Points, const FLinearColor Color, const float Thickness, const bool bDepthTest, const bool bMonotonic, const int32 Segments)
 {
 	const int32 Num = Points.Num();
 	
@@ -1309,47 +1544,51 @@ void UDrawDebugLibrary::DrawDebugCatmullRomSplineArrayView(const FDebugDrawer& D
 	if (Num == 1)
 	{
 		DrawDebugPoint(Drawer, Points[0], Color, Thickness, bDepthTest);
+		return;
 	}
 
 	if (Num == 2)
 	{
 		DrawDebugLine(Drawer, Points[0], Points[1], Color, Thickness, bDepthTest);
+		return;
 	}
 
-	DrawDebugCatmullRomSplineStart(Drawer, Points[0], Points[1], Points[2], bMonotonic, Color, Thickness, Segments, bDepthTest);
+	DrawDebugCatmullRomSplineStart(Drawer, Points[0], Points[1], Points[2], Color, Thickness, bDepthTest, bMonotonic, Segments);
 
 	for (int32 Idx = 0; Idx < Num - 3; Idx++)
 	{
-		DrawDebugCatmullRomSplineSection(Drawer, Points[Idx + 0], Points[Idx + 1], Points[Idx + 2], Points[Idx + 3], bMonotonic, Color, Thickness, Segments, bDepthTest);
+		DrawDebugCatmullRomSplineSection(Drawer, Points[Idx + 0], Points[Idx + 1], Points[Idx + 2], Points[Idx + 3], Color, Thickness, bDepthTest, bMonotonic, Segments);
 	}
 
-	DrawDebugCatmullRomSplineEnd(Drawer, Points[Num - 3], Points[Num - 2], Points[Num - 1], bMonotonic, Color, Thickness, Segments, bDepthTest);
+	DrawDebugCatmullRomSplineEnd(Drawer, Points[Num - 3], Points[Num - 2], Points[Num - 1], Color, Thickness, bDepthTest, bMonotonic, Segments);
 }
 
 
 
-void UDrawDebugLibrary::DrawDebugAngle(const FDebugDrawer& Drawer, const FVector& Location, const FRotator& Rotation, const float Angle, const float LineLength, const float AngleRadius, const FLinearColor Color, const float Thickness, const int32 Segments, const bool bDepthTest)
+void UDrawDebugLibrary::DrawDebugAngle(const FDebugDrawer& Drawer, const FVector& Location, const FRotator& Rotation, const float Angle, const FLinearColor Color, const float Thickness, const bool bDepthTest, const float LineLength, const float AngleRadius, const int32 Segments)
 {
-	DrawDebugArc(Drawer, Location, Rotation, Angle, Color, AngleRadius, Thickness, Segments, bDepthTest);
-	DrawDebugLine(Drawer, Location, Location + Rotation.RotateVector(FVector(LineLength, 0, 0)), Color, Thickness, bDepthTest);
-	DrawDebugLine(Drawer, Location, Location + Rotation.RotateVector(LineLength * FVector(FMath::Cos(FMath::DegreesToRadians(Angle)), FMath::Sin(FMath::DegreesToRadians(Angle)), 0)), Color, Thickness, bDepthTest);
+	DrawDebugArc(Drawer, Location, Rotation, Angle, Color, Thickness, bDepthTest, AngleRadius, Segments);
+	DrawDebugLinesPairsArrayView(Drawer, {
+		{ Location, Location + Rotation.RotateVector(FVector(LineLength, 0, 0)) },
+		{ Location, Location + Rotation.RotateVector(LineLength * FVector(FMath::Cos(FMath::DegreesToRadians(Angle)), FMath::Sin(FMath::DegreesToRadians(Angle)), 0)) }
+		}, Color, Thickness, bDepthTest);
 }
 
-void UDrawDebugLibrary::DrawDebugLocation(const FDebugDrawer& Drawer, const FVector& Location, const FLinearColor Color, const float DrawRadius, const float Thickness, const int32 Segments, const bool bDepthTest)
+void UDrawDebugLibrary::DrawDebugLocation(const FDebugDrawer& Drawer, const FVector& Location, const FLinearColor Color, const float Thickness, const bool bDepthTest, const float DrawRadius, const int32 Segments)
 {
-	DrawDebugSphere(Drawer, Location, FRotator::ZeroRotator, Color, DrawRadius, Thickness, Segments, bDepthTest);
+	DrawDebugSphere(Drawer, Location, FRotator::ZeroRotator, Color,Thickness, bDepthTest, DrawRadius, Segments);
 }
 
-void UDrawDebugLibrary::DrawDebugLocations(const FDebugDrawer& Drawer, const TArray<FVector>& Locations, const FLinearColor Color, const float DrawRadius, const float Thickness, const int32 Segments, const bool bDepthTest)
+void UDrawDebugLibrary::DrawDebugLocations(const FDebugDrawer& Drawer, const TArray<FVector>& Locations, const FLinearColor Color, const float Thickness, const bool bDepthTest, const float DrawRadius, const int32 Segments)
 {
-	DrawDebugLocationsArrayView(Drawer, Locations, Color, DrawRadius, Thickness, Segments, bDepthTest);
+	DrawDebugLocationsArrayView(Drawer, Locations, Color, Thickness, bDepthTest, DrawRadius, Segments);
 }
 
-void UDrawDebugLibrary::DrawDebugLocationsArrayView(const FDebugDrawer& Drawer, const TArrayView<const FVector> Locations, const FLinearColor Color, const float DrawRadius, const float Thickness, const int32 Segments, const bool bDepthTest)
+void UDrawDebugLibrary::DrawDebugLocationsArrayView(const FDebugDrawer& Drawer, const TArrayView<const FVector> Locations, const FLinearColor Color, const float Thickness, const bool bDepthTest, const float DrawRadius, const int32 Segments)
 {
 	for (const FVector& Location : Locations)
 	{
-		DrawDebugLocation(Drawer, Location, Color, DrawRadius, Thickness, Segments, bDepthTest);
+		DrawDebugLocation(Drawer, Location, Color, Thickness, bDepthTest, DrawRadius, Segments);
 	}
 }
 
@@ -1357,27 +1596,27 @@ void UDrawDebugLibrary::DrawDebugRotation(
 	const FDebugDrawer& Drawer,
 	const FVector& Location,
 	const FRotator& Rotation,
-	const float DrawRadius,
-	const float Thickness,
 	const float Opacity,
-	const bool bDepthTest)
+	const float Thickness,
+	const bool bDepthTest,
+	const float DrawRadius)
 {
 	DrawDebugLine(Drawer, Location, Location + Rotation.RotateVector(FVector(DrawRadius, 0.0f, 0.0f)), FLinearColor(1.0f, 0.0f, 0.0f, Opacity), Thickness, bDepthTest);
 	DrawDebugLine(Drawer, Location, Location + Rotation.RotateVector(FVector(0.0f, DrawRadius, 0.0f)), FLinearColor(0.0f, 1.0f, 0.0f, Opacity), Thickness, bDepthTest);
 	DrawDebugLine(Drawer, Location, Location + Rotation.RotateVector(FVector(0.0f, 0.0f, DrawRadius)), FLinearColor(0.0f, 0.0f, 1.0f, Opacity), Thickness, bDepthTest);
 }
 
-void UDrawDebugLibrary::DrawDebugDirection(const FDebugDrawer& Drawer, const FVector& Location, const FVector& Direction, const FLinearColor Color, const float DrawArrowLength, const float ArrowHeadScale, const float Thickness, const bool bDepthTest)
+void UDrawDebugLibrary::DrawDebugDirection(const FDebugDrawer& Drawer, const FVector& Location, const FVector& Direction, const FLinearColor Color, const float Thickness, const bool bDepthTest, const float DrawArrowLength, const float ArrowHeadScale)
 {
-	DrawDebugArrow(Drawer, Location, Location + DrawArrowLength * Direction, Color, 0.2f * ArrowHeadScale * DrawArrowLength, Thickness, bDepthTest);
+	DrawDebugArrow(Drawer, Location, Location + DrawArrowLength * Direction, Color, Thickness, bDepthTest);
 }
 
-void UDrawDebugLibrary::DrawDebugVelocity(const FDebugDrawer& Drawer, const FVector& Location, const FVector& Velocity, const FLinearColor Color, const float DrawVelocityLineScale, const float Thickness, const bool bDepthTest)
+void UDrawDebugLibrary::DrawDebugVelocity(const FDebugDrawer& Drawer, const FVector& Location, const FVector& Velocity, const FLinearColor Color, const float Thickness, const bool bDepthTest, const float DrawVelocityLineScale)
 {
 	DrawDebugLine(Drawer, Location, Location + DrawVelocityLineScale * Velocity, Color, Thickness, bDepthTest);
 }
 
-void UDrawDebugLibrary::DrawDebugTransform(const FDebugDrawer& Drawer, const FTransform& Transform, const float DrawRadius, const float Thickness, const float Opacity, const bool bDepthTest)
+void UDrawDebugLibrary::DrawDebugTransform(const FDebugDrawer& Drawer, const FTransform& Transform, const float Opacity, const float Thickness, const bool bDepthTest, const float DrawRadius)
 {
 	const FVector Location = Transform.GetLocation();
 	const FQuat Rotation = Transform.GetRotation();
@@ -1388,7 +1627,7 @@ void UDrawDebugLibrary::DrawDebugTransform(const FDebugDrawer& Drawer, const FTr
 	DrawDebugLine(Drawer, Location, Location + Rotation.RotateVector(Scale * FVector(0, 0, DrawRadius)), FLinearColor(0.0f, 0.0f, 1.0f, Opacity), Thickness, bDepthTest);
 }
 
-void UDrawDebugLibrary::DrawDebugEvent(const FDebugDrawer& Drawer, bool bTimeUntilEventKnown, const float TimeUntilEvent, const FVector& Location, const FLinearColor Color, const float Size, const float Thickness, const bool bDepthTest)
+void UDrawDebugLibrary::DrawDebugEvent(const FDebugDrawer& Drawer, bool bTimeUntilEventKnown, const float TimeUntilEvent, const FVector& Location, const FLinearColor Color, const float Thickness, const bool bDepthTest, const float Size)
 {
 	DrawDebugLine(
 		Drawer,
@@ -1396,33 +1635,30 @@ void UDrawDebugLibrary::DrawDebugEvent(const FDebugDrawer& Drawer, bool bTimeUnt
 		Location + FVector(0, 0, -Size),
 		Color, Thickness, bDepthTest);
 
-	DrawDebugSphere(Drawer, Location, FRotator::ZeroRotator, Color, Size * 0.1f, Thickness, 8, bDepthTest);
+	DrawDebugSphere(Drawer, Location, FRotator::ZeroRotator, Color, Thickness, bDepthTest, Size * 0.1f, 8);
 
-	// TODO: Don't hard-code apprehension
 	if (bTimeUntilEventKnown && TimeUntilEvent < 1.0f && TimeUntilEvent > -1.0f)
 	{
-		DrawDebugSphere(Drawer, Location + FVector(0, 0, TimeUntilEvent * Size), FRotator::ZeroRotator, Color, Size * 0.25f, Thickness, 8, bDepthTest);
+		DrawDebugSphere(Drawer, Location + FVector(0, 0, TimeUntilEvent * Size), FRotator::ZeroRotator, Color, Thickness, bDepthTest, Size * 0.25f, 8);
 	}
 }
 
 
-
-
-void UDrawDebugLibrary::DrawDebugChair(const FDebugDrawer& Drawer, const FVector& Location, const FRotator& Rotation, const float Width, const float SeatHeight, const float BackHeight, const float BackTilt, const FLinearColor Color, const float Thickness, const bool bDepthTest)
+void UDrawDebugLibrary::DrawDebugChair(const FDebugDrawer& Drawer, const FVector& Location, const FRotator& Rotation, const FLinearColor Color, const float Thickness, const bool bDepthTest, const FDrawDebugChairSettings& Settings)
 {
-	const float HalfWidth = Width / 2.0f;
+	const float HalfWidth = Settings.Width / 2.0f;
 	const FVector F0 = Location + Rotation.RotateVector(FVector(+HalfWidth, +HalfWidth, 0));
 	const FVector F1 = Location + Rotation.RotateVector(FVector(+HalfWidth, -HalfWidth, 0));
 	const FVector F2 = Location + Rotation.RotateVector(FVector(-HalfWidth, -HalfWidth, 0));
 	const FVector F3 = Location + Rotation.RotateVector(FVector(-HalfWidth, +HalfWidth, 0));
 
-	const FVector S0 = Location + Rotation.RotateVector(FVector(+HalfWidth, +HalfWidth, SeatHeight));
-	const FVector S1 = Location + Rotation.RotateVector(FVector(+HalfWidth, -HalfWidth, SeatHeight));
-	const FVector S2 = Location + Rotation.RotateVector(FVector(-HalfWidth, -HalfWidth, SeatHeight));
-	const FVector S3 = Location + Rotation.RotateVector(FVector(-HalfWidth, +HalfWidth, SeatHeight));
+	const FVector S0 = Location + Rotation.RotateVector(FVector(+HalfWidth, +HalfWidth, Settings.SeatHeight));
+	const FVector S1 = Location + Rotation.RotateVector(FVector(+HalfWidth, -HalfWidth, Settings.SeatHeight));
+	const FVector S2 = Location + Rotation.RotateVector(FVector(-HalfWidth, -HalfWidth, Settings.SeatHeight));
+	const FVector S3 = Location + Rotation.RotateVector(FVector(-HalfWidth, +HalfWidth, Settings.SeatHeight));
 
-	const FVector B2 = Location + Rotation.RotateVector(FVector(-HalfWidth - BackTilt, -HalfWidth, SeatHeight + BackHeight));
-	const FVector B3 = Location + Rotation.RotateVector(FVector(-HalfWidth - BackTilt, +HalfWidth, SeatHeight + BackHeight));
+	const FVector B2 = Location + Rotation.RotateVector(FVector(-HalfWidth - Settings.BackTilt, -HalfWidth, Settings.SeatHeight + Settings.BackHeight));
+	const FVector B3 = Location + Rotation.RotateVector(FVector(-HalfWidth - Settings.BackTilt, +HalfWidth, Settings.SeatHeight + Settings.BackHeight));
 
 	DrawDebugLinesPairsArrayView(Drawer, {
 		{ F0, S0 },
@@ -1441,23 +1677,25 @@ void UDrawDebugLibrary::DrawDebugChair(const FDebugDrawer& Drawer, const FVector
 		}, Color, Thickness, bDepthTest);
 }
 
-void UDrawDebugLibrary::DrawDebugDoor(const FDebugDrawer& Drawer, const FVector& Location, const FRotator& Rotation, const float OpenAngle, const float Width, const float Height, const float Inset, const float HandleHeight, const float HandleOffset, const float HandleRadius, const float HandleExtension, const bool bOuterHandle, const bool bInnerHandle, const bool bHandleOnLeft, const bool bDrawEntryArrow, const FLinearColor Color, const float Thickness, const bool bDepthTest)
+void UDrawDebugLibrary::DrawDebugDoor(const FDebugDrawer& Drawer, const FVector& Location, const FRotator& Rotation, const FLinearColor Color, const float Thickness, const bool bDepthTest, const FDrawDebugDoorSettings& Settings)
 {
-	const float HalfWidth = Width / 2.0f;
+	const float HalfWidth = Settings.Width / 2.0f;
+
 	const FVector F0 = Location + Rotation.RotateVector(FVector(0.0f, -HalfWidth, 0.0f));
 	const FVector F1 = Location + Rotation.RotateVector(FVector(0.0f, +HalfWidth, 0.0f));
-	const FVector F2 = Location + Rotation.RotateVector(FVector(0.0f, +HalfWidth, Height));
-	const FVector F3 = Location + Rotation.RotateVector(FVector(0.0f, -HalfWidth, Height));
+	const FVector F2 = Location + Rotation.RotateVector(FVector(0.0f, +HalfWidth, Settings.Height));
+	const FVector F3 = Location + Rotation.RotateVector(FVector(0.0f, -HalfWidth, Settings.Height));
 
-	const FVector LI0 = FVector(0.0f, -HalfWidth + Inset, Inset);
-	const FVector LI1 = FVector(0.0f, +HalfWidth - Inset, Inset);
-	const FVector LI2 = FVector(0.0f, +HalfWidth - Inset, Height - Inset);
-	const FVector LI3 = FVector(0.0f, -HalfWidth + Inset, Height - Inset);
+	const FVector LI0 = FVector(0.0f, -HalfWidth + Settings.Inset, Settings.Inset);
+	const FVector LI1 = FVector(0.0f, +HalfWidth - Settings.Inset, Settings.Inset);
+	const FVector LI2 = FVector(0.0f, +HalfWidth - Settings.Inset, Settings.Height - Settings.Inset);
+	const FVector LI3 = FVector(0.0f, -HalfWidth + Settings.Inset, Settings.Height - Settings.Inset);
 
-	const FVector OpenOffset = bHandleOnLeft ?
-		FVector(0.0f, +HalfWidth - Inset, 0.0f) :
-		FVector(0.0f, -HalfWidth + Inset, 0.0f);
-	const FQuat OpenRotation = FQuat::MakeFromRotationVector(FVector(0.0f, 0.0f, FMath::DegreesToRadians(-OpenAngle)));
+	const FVector OpenOffset = Settings.bHandleOnLeft ?
+		FVector(0.0f, +HalfWidth - Settings.Inset, 0.0f) :
+		FVector(0.0f, -HalfWidth + Settings.Inset, 0.0f);
+
+	const FQuat OpenRotation = FQuat::MakeFromRotationVector(FVector(0.0f, 0.0f, FMath::DegreesToRadians(-Settings.OpenAngle)));
 
 	const FVector I0 = Location + Rotation.RotateVector(OpenRotation.RotateVector(LI0 - OpenOffset) + OpenOffset);
 	const FVector I1 = Location + Rotation.RotateVector(OpenRotation.RotateVector(LI1 - OpenOffset) + OpenOffset);
@@ -1476,41 +1714,42 @@ void UDrawDebugLibrary::DrawDebugDoor(const FDebugDrawer& Drawer, const FVector&
 		{ I3, I0 },
 		}, Color, Thickness, bDepthTest);
 
-	if (bOuterHandle)
+	if (Settings.bOuterHandle)
 	{
-		const FVector LOuterHandle = bHandleOnLeft ?
-			FVector(+HandleRadius + HandleExtension, HandleOffset - Width, HandleHeight) :
-			FVector(+HandleRadius + HandleExtension, Width - HandleOffset, HandleHeight);
+		const FVector LOuterHandle = Settings.bHandleOnLeft ?
+			FVector(+Settings.HandleRadius + Settings.HandleExtension, Settings.HandleOffset - Settings.Width, Settings.HandleHeight) :
+			FVector(+Settings.HandleRadius + Settings.HandleExtension, Settings.Width - Settings.HandleOffset, Settings.HandleHeight);
 		const FVector OuterHandle = Location + Rotation.RotateVector(OpenRotation.RotateVector(LOuterHandle) + OpenOffset);
-		DrawDebugSphere(Drawer, OuterHandle, FRotator::ZeroRotator, Color, HandleRadius, Thickness, 8, bDepthTest);
+
+		DrawDebugSimpleSphere(Drawer, OuterHandle, FRotator::ZeroRotator, Color, Thickness, bDepthTest, Settings.HandleRadius);
 	}
 
-	if (bInnerHandle)
+	if (Settings.bInnerHandle)
 	{
-		const FVector LInnerHandle = bHandleOnLeft ?
-			FVector(-HandleRadius - HandleExtension, HandleOffset - Width, HandleHeight) :
-			FVector(-HandleRadius - HandleExtension, Width - HandleOffset, HandleHeight);
+		const FVector LInnerHandle = Settings.bHandleOnLeft ?
+			FVector(-Settings.HandleRadius - Settings.HandleExtension, Settings.HandleOffset - Settings.Width, Settings.HandleHeight) :
+			FVector(-Settings.HandleRadius - Settings.HandleExtension, Settings.Width - Settings.HandleOffset, Settings.HandleHeight);
 		const FVector InnerHandle = Location + Rotation.RotateVector(OpenRotation.RotateVector(LInnerHandle) + OpenOffset);
-		DrawDebugSphere(Drawer, InnerHandle, FRotator::ZeroRotator, Color, HandleRadius, Thickness, 8, bDepthTest);
+
+		DrawDebugSimpleSphere(Drawer, InnerHandle, FRotator::ZeroRotator, Color, Thickness, bDepthTest, Settings.HandleRadius);
 	}
 
-	if (bDrawEntryArrow)
+	if (Settings.bDrawEntryArrow)
 	{
 		DrawDebugArrow(
 			Drawer,
 			Location + Rotation.RotateVector(FVector(-50.0f, 0.0f, 2.0f)),
 			Location + Rotation.RotateVector(FVector(-10.0f, 0.0f, 2.0f)),
 			Color,
-			20.0f,
 			Thickness,
 			bDepthTest);
 	}
 }
 
-void UDrawDebugLibrary::DrawDebugCamera(const FDebugDrawer& Drawer, const FVector& Location, const FRotator& Rotation, float Scale, float FOVDeg, const FLinearColor Color, const float Thickness, const bool bDepthTest)
+void UDrawDebugLibrary::DrawDebugCamera(const FDebugDrawer& Drawer, const FVector& Location, const FRotator& Rotation, const FLinearColor Color, const float Thickness, const bool bDepthTest, float Scale, float FOVDeg)
 {
-	DrawDebugRotation(Drawer, Location, Rotation, Scale, Thickness, Color.A, bDepthTest);
-	DrawDebugBox(Drawer, Location, Rotation, FVector(Scale, Scale, Scale), Color, Thickness, bDepthTest);
+	DrawDebugRotation(Drawer, Location, Rotation, Color.A, Thickness,bDepthTest, Scale);
+	DrawDebugBox(Drawer, Location, Rotation, Color, Thickness, bDepthTest, FVector(Scale, Scale, Scale));
 
 	// draw "lens" portion
 	const FRotationTranslationMatrix Axes(Rotation, Location);
@@ -1539,12 +1778,12 @@ void UDrawDebugLibrary::DrawDebugCamera(const FDebugDrawer& Drawer, const FVecto
 		}, Color, Thickness, bDepthTest);
 }
 
-void UDrawDebugLibrary::DrawDebugTrajectory(const FDebugDrawer& Drawer, const TArray<FVector>& Locations, const TArray<FVector>& Directions, const FTransform& RelativeTransform, const FLinearColor Color, const float DrawArrowLength, const float PointRadius, const float ArrowHeadScale, const float Thickness, const int32 Segments, const float VerticalOffset, const bool bDepthTest)
+void UDrawDebugLibrary::DrawDebugTrajectory(const FDebugDrawer& Drawer, const TArray<FVector>& Locations, const TArray<FVector>& Directions, const FTransform& RelativeTransform, const FLinearColor Color, const float Thickness, const bool bDepthTest, const float DrawArrowLength, const float PointRadius, const float ArrowHeadScale, const int32 Segments, const float VerticalOffset)
 {
-	DrawDebugTrajectoryFromArrayViews(Drawer, Locations, Directions, RelativeTransform, Color, DrawArrowLength, PointRadius, ArrowHeadScale, Thickness, Segments, VerticalOffset, bDepthTest);
+	DrawDebugTrajectoryFromArrayViews(Drawer, Locations, Directions, RelativeTransform, Color, Thickness, bDepthTest, DrawArrowLength, PointRadius, ArrowHeadScale, Segments, VerticalOffset);
 }
 
-void UDrawDebugLibrary::DrawDebugTrajectoryFromArrayViews(const FDebugDrawer& Drawer, const TArrayView<const FVector> Locations, const TArrayView<const FVector> Directions, const FTransform& RelativeTransform, const FLinearColor Color, const float DrawArrowLength, const float PointRadius, const float ArrowHeadScale, const float Thickness, const int32 Segments, const float VerticalOffset, const bool bDepthTest)
+void UDrawDebugLibrary::DrawDebugTrajectoryFromArrayViews(const FDebugDrawer& Drawer, const TArrayView<const FVector> Locations, const TArrayView<const FVector> Directions, const FTransform& RelativeTransform, const FLinearColor Color, const float Thickness, const bool bDepthTest, const float DrawArrowLength, const float PointRadius, const float ArrowHeadScale, const int32 Segments, const float VerticalOffset)
 {
 	const int32 SampleNum = FMath::Min(Locations.Num(), Directions.Num());
 
@@ -1553,8 +1792,8 @@ void UDrawDebugLibrary::DrawDebugTrajectoryFromArrayViews(const FDebugDrawer& Dr
 		const FVector Location = RelativeTransform.TransformPosition(Locations[SampleIdx]) + FVector(0.0f, 0.0f, VerticalOffset);
 		const FVector Direction = RelativeTransform.TransformVectorNoScale(Directions[SampleIdx]);
 
-		DrawDebugSphere(Drawer, Location, FRotator::ZeroRotator, Color, PointRadius, Thickness, Segments, bDepthTest);
-		DrawDebugArrow(Drawer, Location, Location + DrawArrowLength * Direction, Color, 0.2f * ArrowHeadScale * DrawArrowLength, Thickness, bDepthTest);
+		DrawDebugSimpleSphere(Drawer, Location, FRotator::ZeroRotator, Color, Thickness, bDepthTest, PointRadius, Segments);
+		DrawDebugArrow(Drawer, Location, Location + DrawArrowLength * Direction, Color, Thickness, bDepthTest);
 
 		if (SampleIdx > 0)
 		{
@@ -1564,11 +1803,11 @@ void UDrawDebugLibrary::DrawDebugTrajectoryFromArrayViews(const FDebugDrawer& Dr
 	}
 }
 
-void UDrawDebugLibrary::DrawDebugTransformTrajectory(const FDebugDrawer& Drawer, const FTransformTrajectory& TransformTrajectory, const FLinearColor Color, const float Radius, const float Thickness, const float VerticalOffset, const bool bDepthTest)
+void UDrawDebugLibrary::DrawDebugTransformTrajectory(const FDebugDrawer& Drawer, const FTransformTrajectory& TransformTrajectory, const FLinearColor Color, const float Thickness, const bool bDepthTest, const float Radius, const float VerticalOffset)
 {
 
 	FTransformTrajectorySample TimeZeroSample = TransformTrajectory.GetSampleAtTime(0.0f);
-	DrawDebugSphere(Drawer, TimeZeroSample.Position + FVector(0.0f, 0.0f, VerticalOffset), FRotator::ZeroRotator, Color, Radius, Thickness, 8, bDepthTest);
+	DrawDebugSphere(Drawer, TimeZeroSample.Position + FVector(0.0f, 0.0f, VerticalOffset), FRotator::ZeroRotator, Color, Thickness, bDepthTest, Radius);
 
 	const int32 SampleNum = TransformTrajectory.Samples.Num();
 	for (int32 SampleIdx = 0; SampleIdx < SampleNum; SampleIdx++)
@@ -1576,8 +1815,8 @@ void UDrawDebugLibrary::DrawDebugTransformTrajectory(const FDebugDrawer& Drawer,
 		const FVector Location = TransformTrajectory.Samples[SampleIdx].Position + FVector(0.0f, 0.0f, VerticalOffset);
 		const FQuat Rotation = TransformTrajectory.Samples[SampleIdx].Facing;
 
-		DrawDebugRotation(Drawer, Location, Rotation.Rotator(), Radius, Thickness, Color.A, bDepthTest);
-		DrawDebugString(Drawer, FString::Printf(TEXT("%5.2f"), TransformTrajectory.Samples[SampleIdx].TimeInSeconds), Location + FVector(0.0f, 0.0, 10.0f), FRotator::ZeroRotator, 10.0f, true, Color);
+		DrawDebugRotation(Drawer, Location, Rotation.Rotator(), Color.A, Thickness, bDepthTest, Radius);
+		DrawDebugString(Drawer, FString::Printf(TEXT("%5.2f"), TransformTrajectory.Samples[SampleIdx].TimeInSeconds), Location + FVector(0.0f, 0.0, 10.0f), FRotator::ZeroRotator, Color);
 
 		if (SampleIdx > 0)
 		{
@@ -1587,12 +1826,12 @@ void UDrawDebugLibrary::DrawDebugTransformTrajectory(const FDebugDrawer& Drawer,
 	}
 }
 
-void UDrawDebugLibrary::DrawDebugPose(const FDebugDrawer& Drawer, const TArray<FVector>& BoneLocations, const TArray<FVector>& BoneLinearVelocities, const FTransform& RelativeTransform, const FLinearColor Color, const float DrawVelocityLineScale, const float Thickness, const bool bDepthTest)
+void UDrawDebugLibrary::DrawDebugPose(const FDebugDrawer& Drawer, const TArray<FVector>& BoneLocations, const TArray<FVector>& BoneLinearVelocities, const FTransform& RelativeTransform, const FLinearColor Color, const float Thickness, const bool bDepthTest, const float DrawVelocityLineScale)
 {
-	DrawDebugPoseFromArrayViews(Drawer, BoneLocations, BoneLinearVelocities, RelativeTransform, Color, DrawVelocityLineScale, Thickness, bDepthTest);
+	DrawDebugPoseFromArrayViews(Drawer, BoneLocations, BoneLinearVelocities, RelativeTransform, Color, Thickness, bDepthTest, DrawVelocityLineScale);
 }
 
-void UDrawDebugLibrary::DrawDebugPoseFromArrayViews(const FDebugDrawer& Drawer, const TArrayView<const FVector> BoneLocations, const TArrayView<const FVector> BoneLinearVelocities, const FTransform& RelativeTransform, const FLinearColor Color, const float DrawVelocityLineScale, const float Thickness, const bool bDepthTest)
+void UDrawDebugLibrary::DrawDebugPoseFromArrayViews(const FDebugDrawer& Drawer, const TArrayView<const FVector> BoneLocations, const TArrayView<const FVector> BoneLinearVelocities, const FTransform& RelativeTransform, const FLinearColor Color, const float Thickness, const bool bDepthTest, const float DrawVelocityLineScale)
 {
 	const int32 BoneNum = FMath::Min(BoneLocations.Num(), BoneLinearVelocities.Num());
 
@@ -1601,24 +1840,24 @@ void UDrawDebugLibrary::DrawDebugPoseFromArrayViews(const FDebugDrawer& Drawer,
 		const FVector Location = RelativeTransform.TransformPosition(BoneLocations[BoneIdx]);
 		const FVector Velocity = RelativeTransform.TransformVector(BoneLinearVelocities[BoneIdx]);
 
-		DrawDebugSphere(Drawer, Location, FRotator::ZeroRotator, Color, 10.0f, Thickness, 8, bDepthTest);
+		DrawDebugSphere(Drawer, Location, FRotator::ZeroRotator, Color, Thickness, bDepthTest);
 		DrawDebugLine(Drawer, Location, Location + DrawVelocityLineScale * Velocity, Color, Thickness, bDepthTest);
 	}
 }
 
 
-void UDrawDebugLibrary::DrawDebugVelocities(const FDebugDrawer& Drawer, const TArray<FVector>& Locations, const TArray<FVector>& Velocities, const FLinearColor Color, const float DrawVelocityLineScale, const float Thickness, const bool bDepthTest)
+void UDrawDebugLibrary::DrawDebugVelocities(const FDebugDrawer& Drawer, const TArray<FVector>& Locations, const TArray<FVector>& Velocities, const FLinearColor Color, const float Thickness, const bool bDepthTest, const float DrawVelocityLineScale)
 {
-	DrawDebugVelocitiesArrayView(Drawer, Locations, Velocities, Color, DrawVelocityLineScale, Thickness, bDepthTest);
+	DrawDebugVelocitiesArrayView(Drawer, Locations, Velocities, Color, Thickness, bDepthTest, DrawVelocityLineScale);
 }
 
-void UDrawDebugLibrary::DrawDebugVelocitiesArrayView(const FDebugDrawer& Drawer, const TArrayView<const FVector> Locations, const TArrayView<const FVector> Velocities, const FLinearColor Color, const float DrawVelocityLineScale, const float Thickness, const bool bDepthTest)
+void UDrawDebugLibrary::DrawDebugVelocitiesArrayView(const FDebugDrawer& Drawer, const TArrayView<const FVector> Locations, const TArrayView<const FVector> Velocities, const FLinearColor Color, const float Thickness, const bool bDepthTest, const float DrawVelocityLineScale)
 {
 	const int32 VelocityNum = FMath::Min(Locations.Num(), Velocities.Num());
 
 	for (int32 VelocityIdx = 0; VelocityIdx < VelocityNum; VelocityIdx++)
 	{
-		DrawDebugVelocity(Drawer, Locations[VelocityIdx], Velocities[VelocityIdx], Color, DrawVelocityLineScale, Thickness, bDepthTest);
+		DrawDebugVelocity(Drawer, Locations[VelocityIdx], Velocities[VelocityIdx], Color, Thickness, bDepthTest, DrawVelocityLineScale);
 	}
 }
 
@@ -1628,10 +1867,10 @@ void UDrawDebugLibrary::DrawDebugTransformsArrayView(
 	const TArrayView<const FVector> Locations,
 	const TArrayView<const FQuat4f> Rotations,
 	const TArrayView<const FVector3f> Scales,
-	const float DrawRadius,
-	const float Thickness,
 	const float Opacity,
-	const bool bDepthTest)
+	const float Thickness,
+	const bool bDepthTest,
+	const float DrawRadius)
 {
 	check(Locations.Num() == Rotations.Num());
 	check(Locations.Num() == Scales.Num());
@@ -1647,10 +1886,10 @@ void UDrawDebugLibrary::DrawDebugTransformsArrayView(
 		DrawDebugTransform(
 			Drawer,
 			Transform,
-			DrawRadius,
-			Thickness,
 			Opacity,
-			bDepthTest);
+			Thickness,
+			bDepthTest,
+			DrawRadius);
 	}
 }
 
@@ -1658,10 +1897,10 @@ void UDrawDebugLibrary::DrawDebugRotationsQuatArrayView(
 	const FDebugDrawer& Drawer,
 	const TArrayView<const FVector> Locations,
 	const TArrayView<const FQuat4f> Rotations,
-	const float DrawRadius,
-	const float Thickness,
 	const float Opacity,
-	const bool bDepthTest)
+	const float Thickness,
+	const bool bDepthTest,
+	const float DrawRadius)
 {
 	const int32 Num = FMath::Min(Locations.Num(), Rotations.Num());
 
@@ -1671,87 +1910,165 @@ void UDrawDebugLibrary::DrawDebugRotationsQuatArrayView(
 			Drawer,
 			Locations[Idx],
 			(FRotator)Rotations[Idx].Rotator(),
-			DrawRadius,
-			Thickness,
 			Opacity,
-			bDepthTest);
+			Thickness,
+			bDepthTest,
+			DrawRadius);
 	}
 }
 
-void UDrawDebugLibrary::DrawDebugSkeletonArrayView(
-	const FDebugDrawer& Drawer,
-	const FVector RootLocation,
-	const TArrayView<const FVector> BoneLocations,
-	const TArrayView<const FQuat4f> BoneRotations,
-	const TArrayView<const FVector3f> BoneScales,
-	const TArrayView<const int32> BoneIndices,
-	const FReferenceSkeleton& ReferenceSkeleton,
-	const FLinearColor& Color,
-	const float BoneDrawSize)
+FLinearColor UDrawDebugLibrary::GetDefaultBoneColor()
 {
-	// Skeleton Drawing only possible for PDI
-	if (!Drawer.PDI) { return; }
+	return FLinearColor(0.0f, 0.0f, 0.025f, 1.0f);
+}
 
-	const int32 BoneNum = BoneLocations.Num();
+float UDrawDebugLibrary::GetDefaultBoneRadius()
+{
+	return 1.0f;
+}
 
-	TArray<FTransform> WorldTransforms;
-	TArray<FLinearColor> BoneColorsArray;
-	WorldTransforms.SetNumUninitialized(BoneNum);
-	BoneColorsArray.SetNumUninitialized(BoneNum);
-	for (int32 BoneIdx = 0; BoneIdx < BoneNum; BoneIdx++)
+void UDrawDebugLibrary::DrawDebugBone(const FDebugDrawer& Drawer, const FVector& Location, const FRotator& Rotation, const FLinearColor Color, const float Thickness, const bool bDepthTest, const float Radius, const int32 Segments, const bool bDrawTransform)
+{
+	DrawDebugSimpleSphere(Drawer, Location, Rotation, Color, Thickness, bDepthTest, Radius, Segments);
+	if (bDrawTransform)
+	{
+		DrawDebugRotation(Drawer, Location, Rotation, Color.A, Thickness, bDepthTest, Radius);
+	}
+}
+
+void UDrawDebugLibrary::DrawDebugBoneLink(const FDebugDrawer& Drawer, const FVector& ChildLocation, const FVector& ParentLocation, const FRotator& ParentRotation, const FLinearColor Color, const float Thickness, const bool bDepthTest, const float Radius)
+{
+	const FVector HalfWayLocation = (ChildLocation - ParentLocation) / 2.0f + ParentLocation;
+	const float Length = (ChildLocation - ParentLocation).Length() - 2.0f * Radius;
+
+	if (Length > 0.0f)
 	{
-		WorldTransforms[BoneIdx] = FTransform(
-			((FQuat)BoneRotations[BoneIdx]).GetNormalized(),
-			BoneLocations[BoneIdx],
-			(FVector)BoneScales[BoneIdx]);
-		BoneColorsArray[BoneIdx] = Color;
+		// Does this need to be changed to be oriented using the X forward direction to match Persona?
+		const FVector LocalPosition = ParentRotation.UnrotateVector(ChildLocation - ParentLocation);
+		const FQuat LocalRotation = ParentRotation.Quaternion() * FQuat::FindBetween(FVector::UpVector, LocalPosition);
+
+		DrawDebugSquareBasePyramid(Drawer, HalfWayLocation, LocalRotation.Rotator(), Color, Thickness, bDepthTest, Length, Radius);
 	}
+}
+
+void UDrawDebugLibrary::DrawDebugSkeletonFromSkinnedMeshComponent(
+	const FDebugDrawer& Drawer,
+	const USkinnedMeshComponent* SkinnedMeshComponent,
+	const FLinearColor Color,
+	const float Thickness,
+	const bool bDepthTest,
+	const FDrawDebugSkeletonSettings& Settings)
+{
+	if (!SkinnedMeshComponent) { return; }
 
-	TArray<FBoneIndexType> BoneIndices16;
-	BoneIndices16.SetNumUninitialized(BoneIndices.Num());
-	for (int32 BoneIdx = 0; BoneIdx < BoneIndices.Num(); BoneIdx++)
+	const int32 BoneNum = SkinnedMeshComponent->GetNumBones();
+	
+	const FTransform RootTransform = SkinnedMeshComponent->GetComponentTransform();
+
+	if (Settings.bDrawSimpleSkeleton)
 	{
-		BoneIndices16[BoneIdx] = BoneIndices[BoneIdx];
+		TArray<TPair<FVector, FVector>> LineSegments;
+		LineSegments.Reserve(BoneNum);
+
+		for (int32 BoneIdx = 0; BoneIdx < BoneNum; BoneIdx++)
+		{
+			const FName BoneName = SkinnedMeshComponent->GetBoneName(BoneIdx);
+			const FVector BoneLocation = SkinnedMeshComponent->GetBoneLocation(BoneName, EBoneSpaces::WorldSpace);
+
+			const FName ParentName = SkinnedMeshComponent->GetParentBone(BoneName);
+			if (ParentName != NAME_None)
+			{
+				const FVector ParentLocation = SkinnedMeshComponent->GetBoneLocation(ParentName, EBoneSpaces::WorldSpace);
+
+				LineSegments.Add({ BoneLocation, ParentLocation });
+			}
+			else if (Settings.bDrawRoot && !(RootTransform.GetLocation() - BoneLocation).IsNearlyZero())
+			{
+				DrawDebugLine(Drawer, BoneLocation, RootTransform.GetLocation(), Settings.RootBoneColor, Settings.BoneRadius, bDepthTest);
+			}
+		}
+
+		DrawDebugLinesPairsArrayView(Drawer, LineSegments, Color, Settings.BoneRadius, bDepthTest);
 	}
+	else
+	{
+		for (int32 BoneIdx = 0; BoneIdx < BoneNum; BoneIdx++)
+		{
+			const FName BoneName = SkinnedMeshComponent->GetBoneName(BoneIdx);
+			const FTransform BoneTransform = SkinnedMeshComponent->GetBoneTransform(BoneName, ERelativeTransformSpace::RTS_World);
+
+			DrawDebugBone(Drawer, BoneTransform.GetLocation(), BoneTransform.GetRotation().Rotator(), Color, Thickness, bDepthTest, Settings.BoneRadius, Settings.BoneSegmentNum, Settings.bDrawTransforms);
 
-	FSkelDebugDrawConfig DrawConfig;
-	DrawConfig.BoneDrawMode = EBoneDrawMode::Type::All;
-	DrawConfig.BoneDrawSize = BoneDrawSize;
-	DrawConfig.bAddHitProxy = false;
-	DrawConfig.bForceDraw = false;
-	DrawConfig.bUseMultiColorAsDefaultColor = false;
-	DrawConfig.DefaultBoneColor = Color;
-	DrawConfig.AffectedBoneColor = Color;
-	DrawConfig.SelectedBoneColor = Color;
-	DrawConfig.ParentOfSelectedBoneColor = Color;
+			const FName ParentName = SkinnedMeshComponent->GetParentBone(BoneName);
+			if (ParentName != NAME_None)
+			{
+				const FTransform ParentTransform = SkinnedMeshComponent->GetBoneTransform(ParentName, ERelativeTransformSpace::RTS_World);
 
-	SkeletalDebugRendering::DrawBones(Drawer.PDI, RootLocation, BoneIndices16, ReferenceSkeleton,
-		WorldTransforms, {}, BoneColorsArray, {}, DrawConfig);
+				DrawDebugBoneLink(Drawer, BoneTransform.GetLocation(), ParentTransform.GetLocation(), ParentTransform.GetRotation().Rotator(), Color, Thickness, bDepthTest, Settings.BoneRadius);
+			}
+			else if (Settings.bDrawRoot && !(RootTransform.GetLocation() - BoneTransform.GetLocation()).IsNearlyZero())
+			{
+				DrawDebugBone(Drawer, RootTransform.GetLocation(), RootTransform.GetRotation().Rotator(), Settings.RootBoneColor, Thickness, bDepthTest, Settings.BoneRadius, Settings.BoneSegmentNum, Settings.bDrawTransforms);
+				DrawDebugBoneLink(Drawer, BoneTransform.GetLocation(), RootTransform.GetLocation(), RootTransform.GetRotation().Rotator(), Settings.RootBoneColor, Thickness, bDepthTest, Settings.BoneRadius);
+			}
+		}
+	}
 }
 
-void UDrawDebugLibrary::DrawDebugSimpleSkeletonArrayView(
+void UDrawDebugLibrary::DrawDebugSkeletonArrayView(
 	const FDebugDrawer& Drawer,
+	const FVector RootLocation,
+	const FQuat4f RootRotation,
 	const TArrayView<const FVector> BoneLocations,
+	const TArrayView<const FQuat4f> BoneRotations,
 	const TArrayView<const int32> BoneIndices,
 	const TArrayView<const int32> BoneParents,
-	const FLinearColor& Color,
+	const FLinearColor Color,
 	const float Thickness,
-	const bool bDepthTest)
+	const bool bDepthTest,
+	const FDrawDebugSkeletonSettings& Settings)
 {
-	TArray<TPair<FVector, FVector>> LineSegments;
-	LineSegments.Reserve(BoneIndices.Num());
-
-	for (int32 BoneIdx : BoneIndices)
+	if (Settings.bDrawSimpleSkeleton)
 	{
-		const int32 ParentIdx = BoneParents[BoneIdx];
+		TArray<TPair<FVector, FVector>> LineSegments;
+		LineSegments.Reserve(BoneIndices.Num());
 
-		if (ParentIdx != -1)
+		for (int32 BoneIdx : BoneIndices)
 		{
-			LineSegments.Add({ BoneLocations[BoneIdx], BoneLocations[ParentIdx] });
+			const int32 ParentIdx = BoneParents[BoneIdx];
+
+			if (ParentIdx != INDEX_NONE)
+			{
+				LineSegments.Add({ BoneLocations[BoneIdx], BoneLocations[ParentIdx] });
+			}
+			else if (Settings.bDrawRoot && !(RootLocation - BoneLocations[BoneIdx]).IsNearlyZero())
+			{
+				DrawDebugLine(Drawer, BoneLocations[BoneIdx], RootLocation, Settings.RootBoneColor, Settings.BoneRadius, bDepthTest);
+			}
 		}
+
+		DrawDebugLinesPairsArrayView(Drawer, LineSegments, Color, Settings.BoneRadius, bDepthTest);
 	}
+	else
+	{
 
-	DrawDebugLinesPairsArrayView(Drawer, LineSegments, Color, Thickness, bDepthTest);
+		for (int32 BoneIdx : BoneIndices)
+		{
+			DrawDebugBone(Drawer, BoneLocations[BoneIdx], (FRotator)BoneRotations[BoneIdx].Rotator(), Color, Thickness, bDepthTest, Settings.BoneRadius, Settings.BoneSegmentNum, Settings.bDrawTransforms);
+
+			const int32 ParentIdx = BoneParents[BoneIdx];
+
+			if (ParentIdx != INDEX_NONE)
+			{
+				DrawDebugBoneLink(Drawer, BoneLocations[BoneIdx], BoneLocations[ParentIdx], (FRotator)BoneRotations[ParentIdx].Rotator(), Color, Thickness, bDepthTest, Settings.BoneRadius);
+			}
+			else if (Settings.bDrawRoot && !(RootLocation - BoneLocations[BoneIdx]).IsNearlyZero())
+			{
+				DrawDebugBone(Drawer, RootLocation, (FRotator)RootRotation.Rotator(), Settings.RootBoneColor, Thickness, bDepthTest, Settings.BoneRadius, Settings.BoneSegmentNum, Settings.bDrawTransforms);
+				DrawDebugBoneLink(Drawer, BoneLocations[BoneIdx], RootLocation, (FRotator)RootRotation.Rotator(), Settings.RootBoneColor, Thickness, bDepthTest, Settings.BoneRadius);
+			}
+		}
+	}
 }
 
 void UDrawDebugLibrary::DrawDebugBoneVelocitiesArrayView(
@@ -1760,9 +2077,9 @@ void UDrawDebugLibrary::DrawDebugBoneVelocitiesArrayView(
 	const TArrayView<const FVector3f> BoneVelocities,
 	const TArrayView<const int32> BoneIndices,
 	const FLinearColor& Color,
-	const float Scale,
 	const float Thickness,
-	const bool bDepthTest)
+	const bool bDepthTest,
+	const float Scale)
 {
 	TArray<TPair<FVector, FVector>> LineSegments;
 	LineSegments.Reserve(BoneIndices.Num());
@@ -1779,11 +2096,11 @@ void UDrawDebugLibrary::DrawDebugRangeTrajectoryArrayView(
 	const FDebugDrawer& Drawer,
 	const TArrayView<const FVector> TrajectoryLocations,
 	const TArrayView<const FQuat4f> TrajectoryRotations,
+	const bool bDepthTest,
 	const FVector& ForwardVector,
 	const bool bDrawOrientations,
 	const bool bDrawStartingAtOrigin,
-	const FVector& OriginOffset,
-	const bool bDepthTest)
+	const FVector& OriginOffset)
 {
 	const int32 FrameNum = TrajectoryLocations.Num();
 	if (FrameNum == 0) { return; }
@@ -1828,7 +2145,9 @@ void UDrawDebugLibrary::DrawDebugRangeTrajectoryArrayView(
 					CurrLocation,
 					CurrDirection.Rotation(),
 					FColor::Black.WithAlpha(64),
-					15.0f, 8.0f, 1.0f, bDepthTest);
+					1.0f,
+					bDepthTest,
+					15.0f, 15.0f);
 			}
 		}
 
@@ -1848,48 +2167,12 @@ void UDrawDebugLibrary::DrawDebugRangeTrajectoryArrayView(
 	DrawDebugLinesPairsArrayView(Drawer, LineSegments, FColor::Black.WithAlpha(64), 1.0f, bDepthTest);
 }
 
-void UDrawDebugLibrary::DrawDebugSkeletonCapsulesArrayView(
-	const FDebugDrawer& Drawer,
-	const TArrayView<const FVector> BoneLocations,
-	const TArrayView<const FQuat4f> BoneRotations,
-	const TArrayView<const FVector3f> BoneLocalLocations,
-	const TArrayView<const int32> BoneIndices,
-	const TArrayView<const int32> BoneParents,
-	const float MinCapsuleRadius,
-	const float MaxCapsuleRadius,
-	const FLinearColor Color,
-	const float Thickness,
-	const int32 Segments,
-	const bool bDepthTest)
+int32 UDrawDebugLibrary::DrawDebugStringSegmentNum(const FString& String, const FDrawDebugStringSettings& Settings)
 {
-	for (const int32 BoneIdx : BoneIndices)
-	{
-		const int32 ParentBoneIdx = BoneParents[BoneIdx];
-		if (ParentBoneIdx != -1)
-		{
-			const FVector Location = (BoneLocations[BoneIdx] + BoneLocations[ParentBoneIdx]) / 2.0f;
-			const FQuat4f CapsuleRotation = FQuat4f::FindBetweenNormals(
-				FVector3f(0.0f, 0.0f, 1.0f),
-				BoneLocalLocations[BoneIdx].GetSafeNormal(UE_SMALL_NUMBER, FVector3f::ForwardVector));
-
-			const float Length = (BoneLocations[BoneIdx] - BoneLocations[ParentBoneIdx]).Length();
-			const float Radius = FMath::Clamp(Length / 2.0, MinCapsuleRadius, MaxCapsuleRadius);
-
-			UDrawDebugLibrary::DrawDebugCapsule(
-				Drawer,
-				Location,
-				(FRotator)(BoneRotations[ParentBoneIdx] * CapsuleRotation).Rotator(),
-				Color,
-				Radius,
-				Length * 0.5f,
-				Thickness,
-				Segments,
-				bDepthTest);
-		}
-	}
+	return DrawDebugStringViewSegmentNum(String, Settings);
 }
 
-int32 UDrawDebugLibrary::DrawDebugStringViewSegmentNum(const FStringView String, const bool bMonospaced)
+int32 UDrawDebugLibrary::DrawDebugStringViewSegmentNum(const FStringView String, const FDrawDebugStringSettings& Settings)
 {
 	const int32 CharNum = String.Len();
 
@@ -1900,22 +2183,23 @@ int32 UDrawDebugLibrary::DrawDebugStringViewSegmentNum(const FStringView String,
 		if (String.GetData()[CharIdx] >= '!' && String.GetData()[CharIdx] <= '~')
 		{
 			const int32 TableIdx = String.GetData()[CharIdx] - '!';
-			SegmentNum += bMonospaced ? UE::DrawDebugLibrary::Private::MonoStringLineNums[TableIdx] : UE::DrawDebugLibrary::Private::StringLineNums[TableIdx];
+			SegmentNum += Settings.bMonospaced ? UE::DrawDebugLibrary::Private::MonoStringLineNums[TableIdx] : UE::DrawDebugLibrary::Private::StringLineNums[TableIdx];
 		}
 	}
 
 	return SegmentNum;
 }
 
-FVector UDrawDebugLibrary::DrawDebugStringDimensions(const FString& String, const float Scale, const bool bMonospaced)
+FVector UDrawDebugLibrary::DrawDebugStringDimensions(const FString& String, const FDrawDebugStringSettings& Settings)
 {
-	return DrawDebugStringViewDimensions(String, Scale, bMonospaced);
+	return DrawDebugStringViewDimensions(String, Settings);
 }
 
-FVector UDrawDebugLibrary::DrawDebugStringViewDimensions(const FStringView String, const float Scale, const bool bMonospaced)
+FVector UDrawDebugLibrary::DrawDebugStringViewDimensions(const FStringView String, const FDrawDebugStringSettings& Settings)
 {
 	const int32 CharNum = String.Len();
 
+	float XOffsetMax = 0.0f;
 	float XOffset = 0.0;
 	float YOffset = 0.0;
 
@@ -1923,55 +2207,58 @@ FVector UDrawDebugLibrary::DrawDebugStringViewDimensions(const FStringView Strin
 	{
 		if (String.GetData()[CharIdx] == '\n')
 		{
+			XOffsetMax = FMath::Max(XOffsetMax, XOffset);
 			XOffset = 0.0;
-			YOffset -= Scale;
+			YOffset -= Settings.HeightScale * Settings.LineSpacing * Settings.Height;
 		}
 		else if (String.GetData()[CharIdx] == '\r')
 		{
+			XOffsetMax = FMath::Max(XOffsetMax, XOffset);
 			XOffset = 0.0;
 		}
 		else if (String.GetData()[CharIdx] == '\t')
 		{
-			XOffset += 4.0f * 0.5f * Scale;
+			XOffset += 4.0f * Settings.WidthScale * 0.5f * Settings.Height + Settings.CharacterSpacing;
 		}
 		else if (String.GetData()[CharIdx] < '!' || String.GetData()[CharIdx] > '~')
 		{
-			XOffset += 0.5f * Scale;
+			XOffset += Settings.WidthScale * 0.5f * Settings.Height + Settings.CharacterSpacing;
 		}
 		else
 		{
 			const int32 TableIdx = String.GetData()[CharIdx] - '!';
-			XOffset += (bMonospaced ? 0.5f : (float)UE::DrawDebugLibrary::Private::StringLineAdvs[TableIdx] / 128) * Scale;
+			XOffset += Settings.WidthScale * (Settings.bMonospaced ? 0.5f : (float)UE::DrawDebugLibrary::Private::StringLineAdvs[TableIdx] / 128) * Settings.Height + Settings.CharacterSpacing;
 		}
 	}
 
-	YOffset -= Scale;
+	XOffsetMax = FMath::Max(XOffsetMax, XOffset);
+	YOffset -= Settings.HeightScale * Settings.LineSpacing * Settings.Height;
 
-	return FVector(XOffset, -YOffset, 0);
+	return FVector(XOffsetMax, -YOffset, 0);
 }
 
-FVector UDrawDebugLibrary::DrawDebugTextDimensions(const FText& Text, const float Scale, const bool bMonospaced)
+FVector UDrawDebugLibrary::DrawDebugTextDimensions(const FText& Text, const FDrawDebugStringSettings& Settings)
 {
-	return DrawDebugStringViewDimensions(Text.ToString(), Scale, bMonospaced);
+	return DrawDebugStringViewDimensions(Text.ToString(), Settings);
 }
 
-FVector UDrawDebugLibrary::DrawDebugNameDimensions(const FName& Name, const float Scale, const bool bMonospaced)
+FVector UDrawDebugLibrary::DrawDebugNameDimensions(const FName& Name, const FDrawDebugStringSettings& Settings)
 {
-	return DrawDebugStringViewDimensions(Name.ToString(), Scale, bMonospaced);
+	return DrawDebugStringViewDimensions(Name.ToString(), Settings);
 }
 
 
-void UDrawDebugLibrary::DrawDebugString(const FDebugDrawer& Drawer, const FString& String, const FVector& Location, const FRotator& Rotation, const float Scale, const bool bMonospaced, const FLinearColor Color, const float Thickness, const bool bDepthTest)
+void UDrawDebugLibrary::DrawDebugString(const FDebugDrawer& Drawer, const FString& String, const FVector& Location, const FRotator& Rotation, const FLinearColor Color, const float Thickness, const bool bDepthTest, const FDrawDebugStringSettings& Settings)
 {
-	DrawDebugStringView(Drawer, String, Location, Rotation, Scale, bMonospaced, Color, Thickness, bDepthTest);
+	DrawDebugStringView(Drawer, String, Location, Rotation, Color, Thickness, bDepthTest, Settings);
 }
 
-void UDrawDebugLibrary::DrawDebugStringView(const FDebugDrawer& Drawer, const FStringView String, const FVector& Location, const FRotator& Rotation, const float Scale, const bool bMonospaced, const FLinearColor Color, const float Thickness, const bool bDepthTest)
+void UDrawDebugLibrary::DrawDebugStringView(const FDebugDrawer& Drawer, const FStringView String, const FVector& Location, const FRotator& Rotation, const FLinearColor Color, const float Thickness, const bool bDepthTest, const FDrawDebugStringSettings& Settings)
 {
 	const int32 CharNum = String.Len();
 
 	TArray<TPair<FVector, FVector>> LineSegments;
-	LineSegments.Reserve(DrawDebugStringViewSegmentNum(String, bMonospaced));
+	LineSegments.Reserve(DrawDebugStringViewSegmentNum(String, Settings));
 
 	float XOffset = 0.0;
 	float YOffset = 0.0;
@@ -1981,7 +2268,7 @@ void UDrawDebugLibrary::DrawDebugStringView(const FDebugDrawer& Drawer, const FS
 		if (String.GetData()[CharIdx] == '\n')
 		{
 			XOffset = 0.0;
-			YOffset -= Scale;
+			YOffset -= Settings.LineSpacing * Settings.HeightScale * Settings.Height;
 		}
 		else if (String.GetData()[CharIdx] == '\r')
 		{
@@ -1989,18 +2276,18 @@ void UDrawDebugLibrary::DrawDebugStringView(const FDebugDrawer& Drawer, const FS
 		}
 		else if (String.GetData()[CharIdx] == '\t')
 		{
-			XOffset += 4.0f * 0.5f * Scale;
+			XOffset += 4.0f * Settings.WidthScale * 0.5f * Settings.Height + Settings.CharacterSpacing;
 		}
 		else if (String.GetData()[CharIdx] < '!' || String.GetData()[CharIdx] > '~')
 		{
-			XOffset += 0.5f * Scale;
+			XOffset += Settings.WidthScale * 0.5f * Settings.Height + Settings.CharacterSpacing;
 		}
 		else
 		{
 			const int32 TableIdx = String.GetData()[CharIdx] - '!';
-			const int32 SegmentNum = bMonospaced ? UE::DrawDebugLibrary::Private::MonoStringLineNums[TableIdx] : UE::DrawDebugLibrary::Private::StringLineNums[TableIdx];
-			const int32 SegmentOff = bMonospaced ? UE::DrawDebugLibrary::Private::MonoStringLineOffsets[TableIdx] : UE::DrawDebugLibrary::Private::StringLineOffsets[TableIdx];
-			const int* Segments = bMonospaced ? UE::DrawDebugLibrary::Private::MonoStringLineSegments : UE::DrawDebugLibrary::Private::StringLineSegments;
+			const int32 SegmentNum = Settings.bMonospaced ? UE::DrawDebugLibrary::Private::MonoStringLineNums[TableIdx] : UE::DrawDebugLibrary::Private::StringLineNums[TableIdx];
+			const int32 SegmentOff = Settings.bMonospaced ? UE::DrawDebugLibrary::Private::MonoStringLineOffsets[TableIdx] : UE::DrawDebugLibrary::Private::StringLineOffsets[TableIdx];
+			const int* Segments = Settings.bMonospaced ? UE::DrawDebugLibrary::Private::MonoStringLineSegments : UE::DrawDebugLibrary::Private::StringLineSegments;
 
 			for (int32 SegmentIdx = 0; SegmentIdx < SegmentNum; SegmentIdx++)
 			{
@@ -2010,78 +2297,67 @@ void UDrawDebugLibrary::DrawDebugStringView(const FDebugDrawer& Drawer, const FS
 				const int32 YEndInt = (0x000000FF & (Segments[SegmentOff + SegmentIdx] >> 24));
 
 				const FVector Start = FVector(
-					XOffset + Scale * (((float)XStartInt) / 128),
+					XOffset + Settings.WidthScale * Settings.Height * (((float)XStartInt) / 128),
 					0.0f,
-					YOffset - Scale * (1.0f - ((float)YStartInt) / 128));
+					YOffset - Settings.HeightScale * Settings.Height * (1.0f - ((float)YStartInt) / 128));
 
 				const FVector End = FVector(
-					XOffset + Scale * (((float)XEndInt) / 128),
+					XOffset + Settings.WidthScale * Settings.Height * (((float)XEndInt) / 128),
 					0.0f,
-					YOffset - Scale * (1.0f - ((float)YEndInt) / 128));
+					YOffset - Settings.HeightScale * Settings.Height * (1.0f - ((float)YEndInt) / 128));
 
 				LineSegments.Add({ Rotation.RotateVector(Start) + Location, Rotation.RotateVector(End) + Location });
 			}
 
-			XOffset += (bMonospaced ? 0.5f : (float)UE::DrawDebugLibrary::Private::StringLineAdvs[TableIdx] / 128) * Scale;
+			XOffset += Settings.WidthScale * (Settings.bMonospaced ? 0.5f : (float)UE::DrawDebugLibrary::Private::StringLineAdvs[TableIdx] / 128) * Settings.Height + Settings.CharacterSpacing;
 		}
 	}
 
 	DrawDebugLinesPairsArrayView(Drawer, LineSegments, Color, Thickness, bDepthTest);
 }
 
-void UDrawDebugLibrary::DrawDebugName(const FDebugDrawer& Drawer, const FName& Name, const FVector& Location, const FRotator& Rotation, const float Scale, const bool bMonospaced, const FLinearColor Color, const float Thickness, const bool bDepthTest)
+void UDrawDebugLibrary::DrawDebugName(const FDebugDrawer& Drawer, const FName& Name, const FVector& Location, const FRotator& Rotation, const FLinearColor Color, const float Thickness, const bool bDepthTest, const FDrawDebugStringSettings& Settings)
 {
-	DrawDebugStringView(Drawer, Name.ToString(), Location, Rotation, Scale, bMonospaced, Color, Thickness, bDepthTest);
+	DrawDebugStringView(Drawer, Name.ToString(), Location, Rotation, Color, Thickness, bDepthTest, Settings);
 }
 
-void UDrawDebugLibrary::DrawDebugText(const FDebugDrawer& Drawer, const FText& Text, const FVector& Location, const FRotator& Rotation, const float Scale, const bool bMonospaced, const FLinearColor Color, const float Thickness, const bool bDepthTest)
+void UDrawDebugLibrary::DrawDebugText(const FDebugDrawer& Drawer, const FText& Text, const FVector& Location, const FRotator& Rotation, const FLinearColor Color, const float Thickness, const bool bDepthTest, const FDrawDebugStringSettings& Settings)
 {
-	DrawDebugStringView(Drawer, Text.ToString(), Location, Rotation, Scale, bMonospaced, Color, Thickness, bDepthTest);
+	DrawDebugStringView(Drawer, Text.ToString(), Location, Rotation, Color, Thickness, bDepthTest, Settings);
 }
 
 
-void UDrawDebugLibrary::DrawDebugVisualLoggerString(const FDebugDrawer& Drawer, const FString& String, const FVector& Location, const FLinearColor Color)
+void UDrawDebugLibrary::VisualLoggerDrawString(const FDebugDrawer& Drawer, const FString& String, const FVector& Location, const FLinearColor Color)
 {
-	DrawDebugVisualLoggerStringView(Drawer, String, Location, Color);
+	Drawer.VisualLoggerDrawString(String, Location, Color);
 }
 
-void UDrawDebugLibrary::DrawDebugVisualLoggerStringView(const FDebugDrawer& Drawer, const FStringView String, const FVector& Location, const FLinearColor Color)
+void UDrawDebugLibrary::VisualLoggerDrawStringView(const FDebugDrawer& Drawer, const FStringView String, const FVector& Location, const FLinearColor Color)
 {
-#if ENABLE_VISUAL_LOG
-	if (Drawer.VisualLoggerObject)
-	{
-		FVisualLogEntry* CurrentEntry = nullptr;
-		if (FVisualLogger::CheckVisualLogInputInternal(Drawer.VisualLoggerObject, Drawer.VisualLoggerCategory, Drawer.VisualLoggerVerbosity, /*OutWorld=*/ nullptr, &CurrentEntry))
-		{
-			FVisualLogShapeElement Element(FString(String), Color.ToFColor(false), 0, Drawer.VisualLoggerCategory);
-			Element.Type = EVisualLoggerShapeElement::SinglePoint;
-			Element.Verbosity = Drawer.VisualLoggerVerbosity;
-			Element.Points.Add(Location);
-			CurrentEntry->ElementsToDraw.Add(Element);
-		}
-	}
-#endif
+	Drawer.VisualLoggerDrawString(String, Location, Color);
 }
 
-void UDrawDebugLibrary::DrawDebugVisualLoggerName(const FDebugDrawer& Drawer, const FName& Name, const FVector& Location, const FLinearColor Color)
+void UDrawDebugLibrary::VisualLoggerDrawName(const FDebugDrawer& Drawer, const FName& Name, const FVector& Location, const FLinearColor Color)
 {
-	DrawDebugVisualLoggerStringView(Drawer, Name.ToString(), Location, Color);
+	Drawer.VisualLoggerDrawString(Name.ToString(), Location, Color);
 }
 
-void UDrawDebugLibrary::DrawDebugVisualLoggerText(const FDebugDrawer& Drawer, const FText& Text, const FVector& Location, const FLinearColor Color)
+void UDrawDebugLibrary::VisualLoggerDrawText(const FDebugDrawer& Drawer, const FText& Text, const FVector& Location, const FLinearColor Color)
 {
-	DrawDebugVisualLoggerStringView(Drawer, Text.ToString(), Location, Color);
+	Drawer.VisualLoggerDrawString(Text.ToString(), Location, Color);
 }
 
 
-void UDrawDebugLibrary::DrawDebugMoverOrientation(const FDebugDrawer& Drawer, const FVector& Location, const FRotator& Rotation, const FVector ForwardVector, const float Scale, const FLinearColor Color, const float Thickness, const bool bDepthTest)
+void UDrawDebugLibrary::DrawDebugMoverOrientation(const FDebugDrawer& Drawer, const FVector& Location, const FRotator& Rotation, const FLinearColor Color, const float Thickness, const bool bDepthTest, const FVector ForwardVector, const float Scale)
 {
 	DrawDebugRotation(
 		Drawer,
 		Location,
 		Rotation,
-		Scale / 3.0f,
-		Thickness);
+		Color.A,
+		Thickness,
+		true,
+		Scale / 3.0f);
 
 	const float ArrowSize = (2.0f / 3.0f) * Scale;
 	const FVector ArrowDirection = Rotation.RotateVector(ForwardVector);
@@ -2092,17 +2368,19 @@ void UDrawDebugLibrary::DrawDebugMoverOrientation(const FDebugDrawer& Drawer, co
 		ArrowPosition,
 		ArrowDirection.Rotation(),
 		Color,
+		Thickness,
+		true,
 		2.0f * ArrowSize,
-		ArrowSize,
-		Thickness);
+		2.0f * ArrowSize);
 
 	DrawDebugCircle(
 		Drawer,
 		Location,
 		FRotator::ZeroRotator,
 		Color,
-		Scale,
 		Thickness,
+		true,
+		Scale,
 		31);
 }
 
@@ -2118,17 +2396,14 @@ void UDrawDebugLibrary::DrawDebugGraph(
 	const float Ymax,
 	const float XaxisLength,
 	const float YaxisLength,
-	const FString& Title,
-	const FString& XaxisLabel,
-	const FString& YaxisLabel,
-	const float TitleScale,
 	const FLinearColor Color,
 	const float AxisThickness,
 	const float LineThickness,
 	const float TextThickness,
-	const bool bDepthTest)
+	const bool bDepthTest,
+	const FDrawDebugGraphAxesSettings& AxesSettings)
 {
-	DrawDebugGraphView(
+	DrawDebugGraphArrayView(
 		Drawer,
 		Location,
 		Rotation,
@@ -2140,18 +2415,15 @@ void UDrawDebugLibrary::DrawDebugGraph(
 		Ymax,
 		XaxisLength,
 		YaxisLength,
-		Title,
-		XaxisLabel,
-		YaxisLabel,
-		TitleScale,
 		Color,
 		AxisThickness,
 		LineThickness,
 		TextThickness,
-		bDepthTest);
+		bDepthTest,
+		AxesSettings);
 }
 
-void UDrawDebugLibrary::DrawDebugGraphView(
+void UDrawDebugLibrary::DrawDebugGraphArrayView(
 	const FDebugDrawer& Drawer,
 	const FVector& Location,
 	const FRotator& Rotation,
@@ -2163,15 +2435,53 @@ void UDrawDebugLibrary::DrawDebugGraphView(
 	const float Ymax,
 	const float XaxisLength,
 	const float YaxisLength,
-	const FStringView Title,
-	const FStringView XaxisLabel,
-	const FStringView YaxisLabel,
-	const float TextScale,
 	const FLinearColor Color,
 	const float AxisThickness,
 	const float LineThickness,
 	const float TextThickness,
-	const bool bDepthTest)
+	const bool bDepthTest,
+	const FDrawDebugGraphAxesSettings& AxesSettings)
+{
+	DrawDebugGraphAxes(
+		Drawer,
+		Location,
+		Rotation,
+		XaxisLength,
+		YaxisLength,
+		Color,
+		AxisThickness,
+		TextThickness,
+		bDepthTest,
+		AxesSettings);
+
+	DrawDebugGraphLineArrayView(
+		Drawer,
+		Location,
+		Rotation,
+		Xvalues,
+		Yvalues,
+		Xmin,
+		Xmax,
+		Ymin,
+		Ymax,
+		XaxisLength,
+		YaxisLength,
+		Color,
+		LineThickness,
+		bDepthTest);
+}
+
+void UDrawDebugLibrary::DrawDebugGraphAxes(
+	const FDebugDrawer& Drawer,
+	const FVector& Location,
+	const FRotator& Rotation,
+	const float XaxisLength,
+	const float YaxisLength,
+	const FLinearColor Color,
+	const float AxisThickness,
+	const float TextThickness,
+	const bool bDepthTest,
+	const FDrawDebugGraphAxesSettings& AxesSettings)
 {
 	DrawDebugOrientedArrow(
 		Drawer,
@@ -2179,61 +2489,92 @@ void UDrawDebugLibrary::DrawDebugGraphView(
 		(Rotation.Quaternion() * FQuat::MakeFromRotationVector(FVector(UE_PI / 2, 0, 0))).Rotator(),
 		XaxisLength,
 		Color,
-		10.0f,
 		AxisThickness,
 		bDepthTest);
 
 	DrawDebugOrientedArrow(
 		Drawer,
 		Location,
-		(Rotation.Quaternion() * 
-			FQuat::MakeFromRotationVector(FVector(0, -UE_PI/2, 0)) *
+		(Rotation.Quaternion() *
+			FQuat::MakeFromRotationVector(FVector(0, -UE_PI / 2, 0)) *
 			FQuat::MakeFromRotationVector(FVector(UE_PI / 2, 0, 0))).Rotator(),
 		YaxisLength,
 		Color,
-		10.0f,
 		AxisThickness,
 		bDepthTest);
 
-	const FVector TitleDimension = DrawDebugStringViewDimensions(Title, TextScale, false);
+	const FVector TitleDimension = DrawDebugStringViewDimensions(AxesSettings.Title, AxesSettings.TitleSettings);
 
 	DrawDebugStringView(
 		Drawer,
-		Title,
-		Location + Rotation.RotateVector(FVector(XaxisLength / 2.0f - TitleDimension.X / 2.0f, 0.0, YaxisLength + TitleDimension.Y + 0.1f * TextScale)),
+		AxesSettings.Title,
+		Location + Rotation.RotateVector(FVector(XaxisLength / 2.0f - TitleDimension.X / 2.0f, 0.0, YaxisLength + 1.1f * TitleDimension.Y)),
 		Rotation,
-		TextScale,
-		false,
 		Color,
 		TextThickness,
-		bDepthTest);
+		bDepthTest,
+		AxesSettings.TitleSettings);
 
-	const FVector XlabelDimension = DrawDebugStringViewDimensions(XaxisLabel, TextScale, false);
+	const FVector XlabelDimension = DrawDebugStringViewDimensions(AxesSettings.XaxisLabel, AxesSettings.AxisLabelSettings);
 
 	DrawDebugStringView(
 		Drawer,
-		XaxisLabel,
-		Location + Rotation.RotateVector(FVector(XaxisLength / 2.0f - XlabelDimension.X / 2.0f, 0.0, -0.1f * TextScale)),
+		AxesSettings.XaxisLabel,
+		Location + Rotation.RotateVector(FVector(XaxisLength / 2.0f - XlabelDimension.X / 2.0f, 0.0, -0.1f * XlabelDimension.Y)),
 		Rotation,
-		TextScale,
-		false,
 		Color,
 		TextThickness,
-		bDepthTest);
+		bDepthTest,
+		AxesSettings.AxisLabelSettings);
 
-	const FVector YlabelDimension = DrawDebugStringViewDimensions(YaxisLabel, TextScale, false);
+	const FVector YlabelDimension = DrawDebugStringViewDimensions(AxesSettings.YaxisLabel, AxesSettings.AxisLabelSettings);
 
 	DrawDebugStringView(
 		Drawer,
-		YaxisLabel,
-		Location + Rotation.RotateVector(FVector(-YlabelDimension.Y - 0.1f * TextScale, 0.0, YaxisLength / 2 - YlabelDimension.X / 2)),
-		(Rotation.Quaternion() * FQuat::MakeFromRotationVector(FVector(0, -UE_PI/2, 0))).Rotator(),
-		TextScale,
-		false,
+		AxesSettings.YaxisLabel,
+		Location + Rotation.RotateVector(FVector(-YlabelDimension.Y - 0.1f * YlabelDimension.Y, 0.0, YaxisLength / 2 - YlabelDimension.X / 2)),
+		(Rotation.Quaternion() * FQuat::MakeFromRotationVector(FVector(0, -UE_PI / 2, 0))).Rotator(),
 		Color,
 		TextThickness,
-		bDepthTest);
+		bDepthTest,
+		AxesSettings.AxisLabelSettings);
+}
+
+void UDrawDebugLibrary::DrawDebugGraphLine(
+	const FDebugDrawer& Drawer,
+	const FVector& Location,
+	const FRotator& Rotation,
+	const TArray<float>& Xvalues,
+	const TArray<float>& Yvalues,
+	const float Xmin,
+	const float Xmax,
+	const float Ymin,
+	const float Ymax,
+	const float XaxisLength,
+	const float YaxisLength,
+	const FLinearColor Color,
+	const float LineThickness,
+	const bool bDepthTest)
+{
+	DrawDebugGraphLineArrayView(Drawer, Location, Rotation, Xvalues, Yvalues, Xmin, Xmax, Ymin, Ymax, XaxisLength, YaxisLength, Color, LineThickness, bDepthTest);
+}
 
+void UDrawDebugLibrary::DrawDebugGraphLineArrayView(
+	const FDebugDrawer& Drawer,
+	const FVector& Location,
+	const FRotator& Rotation,
+	const TArrayView<const float> Xvalues,
+	const TArrayView<const float> Yvalues,
+	const float Xmin,
+	const float Xmax,
+	const float Ymin,
+	const float Ymax,
+	const float XaxisLength,
+	const float YaxisLength,
+	const FLinearColor Color,
+	const float LineThickness,
+	const bool bDepthTest)
+{
 	const int32 PointNum = FMath::Min(Xvalues.Num(), Yvalues.Num());
 
 	TArray<TPair<FVector, FVector>> LineSegments;
@@ -2252,4 +2593,68 @@ void UDrawDebugLibrary::DrawDebugGraphView(
 	}
 
 	DrawDebugLinesPairsArrayView(Drawer, LineSegments, Color, LineThickness, bDepthTest);
+}
+
+void UDrawDebugLibrary::DrawDebugGraphLegend(
+	const FDebugDrawer& Drawer,
+	const FVector& Location,
+	const FRotator& Rotation,
+	const TArray<FLinearColor>& LegendColors,
+	const TArray<FString>& LegendLabels,
+	const float XaxisLength,
+	const float YaxisLength,
+	const float TextThickness,
+	const float PointSize,
+	const bool bDepthTest,
+	const FDrawDebugStringSettings& LegendSettings)
+{
+	DrawDebugGraphLegendArrayView(Drawer, Location, Rotation, LegendColors, LegendLabels, XaxisLength, YaxisLength, TextThickness, PointSize, bDepthTest, LegendSettings);
+}
+
+void UDrawDebugLibrary::DrawDebugGraphLegendArrayView(
+	const FDebugDrawer& Drawer,
+	const FVector& Location,
+	const FRotator& Rotation,
+	const TArrayView<const FLinearColor> LegendColors,
+	const TArrayView<const FString> LegendLabels,
+	const float XaxisLength,
+	const float YaxisLength,
+	const float TextThickness,
+	const float PointSize,
+	const bool bDepthTest,
+	const FDrawDebugStringSettings& LegendSettings)
+{
+	const int32 LegendNum = FMath::Min(LegendColors.Num(), LegendLabels.Num());
+
+	float TotalHeight = 0.0f;
+	for (int32 LegendIdx = 0; LegendIdx < LegendNum; LegendIdx++)
+	{
+		TotalHeight = DrawDebugStringViewDimensions(LegendLabels[LegendIdx], LegendSettings).Y;
+	}
+
+	float VerticalOffset = TotalHeight / 2.0f + YaxisLength / 2.0f;
+	for (int32 LegendIdx = 0; LegendIdx < LegendNum; LegendIdx++)
+	{
+		const float LegendHeight = DrawDebugStringViewDimensions(LegendLabels[LegendIdx], LegendSettings).Y;
+
+		const FQuat IconRotation = FQuat::MakeFromRotationVector(FVector(UE_PI / 2.0, 0.0f, 0.0f));
+
+		DrawDebugDiamond(
+			Drawer, 
+			Location + Rotation.RotateVector(FVector(XaxisLength + 2.0f * PointSize, 0.0, VerticalOffset - LegendHeight / 2.0f)),
+			(Rotation.Quaternion() * IconRotation).Rotator(),
+			LegendColors[LegendIdx], TextThickness, bDepthTest, PointSize / 2.0f);
+
+		DrawDebugStringView(
+			Drawer,
+			LegendLabels[LegendIdx],
+			Location + Rotation.RotateVector(FVector(XaxisLength + 4.0f * PointSize, 0.0, VerticalOffset)),
+			Rotation,
+			LegendColors[LegendIdx],
+			TextThickness,
+			bDepthTest,
+			LegendSettings);
+
+		VerticalOffset += LegendHeight;
+	}
 }
\ No newline at end of file
diff --git a/Engine/Plugins/Experimental/DrawDebugLibrary/Source/DrawDebugLibrary/Public/DrawDebugLibrary.h b/Engine/Plugins/Experimental/DrawDebugLibrary/Source/DrawDebugLibrary/Public/DrawDebugLibrary.h
--- a/Engine/Plugins/Experimental/DrawDebugLibrary/Source/DrawDebugLibrary/Public/DrawDebugLibrary.h
+++ b/Engine/Plugins/Experimental/DrawDebugLibrary/Source/DrawDebugLibrary/Public/DrawDebugLibrary.h
@@ -11,9 +11,9 @@
 #define UE_API DRAWDEBUGLIBRARY_API
 
 class FPrimitiveDrawInterface;
-struct FReferenceSkeleton;
 struct FTransformTrajectory;
 struct FAnimInstanceProxy;
+class USkinnedMeshComponent;
 
 class FDrawDebugLibraryModule : public IModuleInterface
 {
@@ -25,7 +25,7 @@ public:
 };
 
 /** Blueprint accessible version of log verbosity for the Visual Logger Drawer */
-UENUM(BlueprintType)
+UENUM(BlueprintType, Category = "Draw Debug Library")
 enum class EDrawDebugLogVerbosity : uint8
 {
 	/** Always prints a fatal error to console (and log file) and crashes (even if logging is disabled) */
@@ -62,29 +62,300 @@ enum class EDrawDebugLogVerbosity : uint8
 	VeryVerbose,
 };
 
+/** Enum for different arrow head types */
+UENUM(BlueprintType, Category = "Draw Debug Library")
+enum class EDrawDebugArrowHead : uint8
+{
+	/** Simple pointed arrow */
+	Simple,
+
+	/** Four-way arrow head */
+	FourWay,
+
+	/** Triangular-based pyramid arrow head */
+	Triangular,
+
+	/** Square-based pyramid arrow head */
+	Square,
+
+	/** Cone arrow head */
+	Cone,
+};
+
+
+/** Settings for debug drawing arrows */
+USTRUCT(BlueprintType, Category = "Draw Debug Library")
+struct FDrawDebugArrowSettings
+{
+	GENERATED_BODY()
+
+	/** If to draw an arrow head on the start of the arrow */
+	UPROPERTY(BlueprintReadWrite, EditAnywhere, Category = "Draw Debug Library")
+	bool bArrowheadOnStart = false;
+
+	/** Width of the start arrow head */
+	UPROPERTY(BlueprintReadWrite, EditAnywhere, Category = "Draw Debug Library", meta = (EditCondition = "bArrowheadOnStart", HideEditConditionToggle, ClampMin = "0", UIMin = "0", Units = "cm"))
+	float ArrowHeadStartWidth = 5.0f;
+
+	/** Type of start arrow head */
+	UPROPERTY(BlueprintReadWrite, EditAnywhere, Category = "Draw Debug Library", meta = (EditCondition = "bArrowheadOnStart", HideEditConditionToggle))
+	EDrawDebugArrowHead ArrowHeadStartType = EDrawDebugArrowHead::Simple;
+
+	/** If to draw an arrow head on the end of the arrow */
+	UPROPERTY(BlueprintReadWrite, EditAnywhere, Category = "Draw Debug Library")
+	bool bArrowheadOnEnd = true;
+
+	/** Width of the end arrow head */
+	UPROPERTY(BlueprintReadWrite, EditAnywhere, Category = "Draw Debug Library", meta = (EditCondition = "bArrowheadOnEnd", HideEditConditionToggle, ClampMin = "0", UIMin = "0", Units = "cm"))
+	float ArrowHeadEndWidth = 5.0f;
+
+	/** Type of end arrow head */
+	UPROPERTY(BlueprintReadWrite, EditAnywhere, Category = "Draw Debug Library", meta = (EditCondition = "bArrowheadOnEnd", HideEditConditionToggle))
+	EDrawDebugArrowHead ArrowHeadEndType = EDrawDebugArrowHead::Simple;
+};
+
+/** Settings for debug drawing chairs */
+USTRUCT(BlueprintType, Category = "Draw Debug Library")
+struct FDrawDebugChairSettings
+{
+	GENERATED_BODY()
+
+	/** Width of the chair */
+	UPROPERTY(BlueprintReadWrite, EditAnywhere, Category = "Draw Debug Library", meta = (ClampMin = "0", UIMin = "0", Units = "cm"))
+	float Width = 40.0f;
+
+	/** Height of the seat of the chair off the ground */
+	UPROPERTY(BlueprintReadWrite, EditAnywhere, Category = "Draw Debug Library", meta = (ClampMin = "0", UIMin = "0", Units = "cm"))
+	float SeatHeight = 40.0f;
+
+	/** Height of the back of the chair off the seat */
+	UPROPERTY(BlueprintReadWrite, EditAnywhere, Category = "Draw Debug Library", meta = (ClampMin = "0", UIMin = "0", Units = "cm"))
+	float BackHeight = 40.0f;
+
+	/** Tilt of the chair backwards in cm */
+	UPROPERTY(BlueprintReadWrite, EditAnywhere, Category = "Draw Debug Library", meta = (Units = "cm"))
+	float BackTilt = 10.0f;
+};
+
+/** Settings for debug drawing doors */
+USTRUCT(BlueprintType, Category = "Draw Debug Library")
+struct FDrawDebugDoorSettings
+{
+	GENERATED_BODY()
+
+	/** Width of the door */
+	UPROPERTY(BlueprintReadWrite, EditAnywhere, Category = "Draw Debug Library", meta = (ClampMin = "0", UIMin = "0", Units = "cm"))
+	float Width = 100.0f;
+
+	/** Height of the door */
+	UPROPERTY(BlueprintReadWrite, EditAnywhere, Category = "Draw Debug Library", meta = (ClampMin = "0", UIMin = "0", Units = "cm"))
+	float Height = 200.0f;
+
+	/** How much the door is inset from the frame */
+	UPROPERTY(BlueprintReadWrite, EditAnywhere, Category = "Draw Debug Library", meta = (ClampMin = "0", UIMin = "0", Units = "cm"))
+	float Inset = 5.0f;
+
+	/** The height of the handle off the ground */
+	UPROPERTY(BlueprintReadWrite, EditAnywhere, Category = "Draw Debug Library", meta = (ClampMin = "0", UIMin = "0", Units = "cm"))
+	float HandleHeight = 90.0f;
+
+	/** The offset of the handle into the door */
+	UPROPERTY(BlueprintReadWrite, EditAnywhere, Category = "Draw Debug Library", meta = (ClampMin = "0", UIMin = "0", Units = "cm"))
+	float HandleOffset = 20.0f;
+
+	/** The radius of the handle */
+	UPROPERTY(BlueprintReadWrite, EditAnywhere, Category = "Draw Debug Library", meta = (ClampMin = "0", UIMin = "0", Units = "cm"))
+	float HandleRadius = 3.0f;
+
+	/** How far the handle extends outward from the door */
+	UPROPERTY(BlueprintReadWrite, EditAnywhere, Category = "Draw Debug Library", meta = (ClampMin = "0", UIMin = "0", Units = "cm"))
+	float HandleExtension = 2.0f;
+
+	/** If to place a handle on the outside entrance of the door */
+	UPROPERTY(BlueprintReadWrite, EditAnywhere, Category = "Draw Debug Library")
+	bool bOuterHandle = true;
+
+	/** If to place a handle on the inside entrance of the door */
+	UPROPERTY(BlueprintReadWrite, EditAnywhere, Category = "Draw Debug Library")
+	bool bInnerHandle = true;
+
+	/** If to place the handle on the left-hand-side or the right-hand-side of the door from the perspective of entering */
+	UPROPERTY(BlueprintReadWrite, EditAnywhere, Category = "Draw Debug Library")
+	bool bHandleOnLeft = true;
+
+	/** If to draw an arrow on the floor showing the natural door entry direction */
+	UPROPERTY(BlueprintReadWrite, EditAnywhere, Category = "Draw Debug Library")
+	bool bDrawEntryArrow = true;
+
+	/** How open the door is in degrees */
+	UPROPERTY(BlueprintReadWrite, EditAnywhere, Category = "Draw Debug Library", meta = (Units = "deg"))
+	float OpenAngle = 40.0f;
+};
+
+/** Settings for debug drawing skeletons */
+USTRUCT(BlueprintType, Category = "Draw Debug Library")
+struct FDrawDebugSkeletonSettings
+{
+	GENERATED_BODY()
+
+	/** Draws a simple skeleton made up of lines rather than the full skeleton */
+	UPROPERTY(BlueprintReadWrite, EditAnywhere, Category = "Draw Debug Library")
+	bool bDrawSimpleSkeleton = false;
+
+	/** Radius of the bones */
+	UPROPERTY(BlueprintReadWrite, EditAnywhere, Category = "Draw Debug Library", meta = (ClampMin = "0", UIMin = "0", Units = "cm"))
+	float BoneRadius = 1.0f;
+
+	/** Number of segments in the spheres drawn for the bones */
+	UPROPERTY(BlueprintReadWrite, EditAnywhere, Category = "Draw Debug Library", meta = (EditCondition = "!bDrawSimpleSkeleton", HideEditConditionToggle, ClampMin = "0", UIMin = "0"))
+	int32 BoneSegmentNum = 10;
+
+	/** If to draw small transforms inside the bone spheres */
+	UPROPERTY(BlueprintReadWrite, EditAnywhere, Category = "Draw Debug Library", meta = (EditCondition = "!bDrawSimpleSkeleton", HideEditConditionToggle))
+	bool bDrawTransforms = false;
+
+	/** If to draw a bone going to the root/component transform */
+	UPROPERTY(BlueprintReadWrite, EditAnywhere, Category = "Draw Debug Library")
+	bool bDrawRoot = true;
+
+	/** Color of the root bone when drawn */
+	UPROPERTY(BlueprintReadWrite, EditAnywhere, Category = "Draw Debug Library", meta = (EditCondition = "bDrawRoot", HideEditConditionToggle))
+	FLinearColor RootBoneColor = FLinearColor::Red;
+};
+
+/** Settings for debug drawing strings */
+USTRUCT(BlueprintType, Category = "Draw Debug Library")
+struct FDrawDebugStringSettings
+{
+	GENERATED_BODY()
+
+	/** Height of the text */
+	UPROPERTY(BlueprintReadWrite, EditAnywhere, Category = "Draw Debug Library", meta = (ClampMin = "0", UIMin = "0", Units = "cm"))
+	float Height = 10.0f;
+
+	/** If to draw using a monospaced font */
+	UPROPERTY(BlueprintReadWrite, EditAnywhere, Category = "Draw Debug Library")
+	bool bMonospaced = true;
+
+	/** Scaling factor for the width of the characters */
+	UPROPERTY(BlueprintReadWrite, EditAnywhere, Category = "Draw Debug Library", meta = (ClampMin = "0", UIMin = "0"))
+	float WidthScale = 1.0f;
+
+	/** Scaling factor for the height of the characters */
+	UPROPERTY(BlueprintReadWrite, EditAnywhere, Category = "Draw Debug Library", meta = (ClampMin = "0", UIMin = "0"))
+	float HeightScale = 1.0f;
+
+	/** Scaling factor for the spacing between lines */
+	UPROPERTY(BlueprintReadWrite, EditAnywhere, Category = "Draw Debug Library", meta = (ClampMin = "0", UIMin = "0"))
+	float LineSpacing = 1.0f;
+
+	/** Gaps between characters in cm */
+	UPROPERTY(BlueprintReadWrite, EditAnywhere, Category = "Draw Debug Library", meta = (ClampMin = "0", UIMin = "0", Units = "cm"))
+	float CharacterSpacing = 0.0f;
+};
+
+/** Settings for debug drawing graphs axes */
+USTRUCT(BlueprintType, Category = "Draw Debug Library")
+struct FDrawDebugGraphAxesSettings
+{
+	GENERATED_BODY()
+
+	/** Title of the graph */
+	UPROPERTY(BlueprintReadWrite, EditAnywhere, Category = "Draw Debug Library")
+	FString Title;
+
+	/** Settings for the title */
+	UPROPERTY(BlueprintReadWrite, EditAnywhere, Category = "Draw Debug Library")
+	FDrawDebugStringSettings TitleSettings;
+
+	/** Label for the X axis */
+	UPROPERTY(BlueprintReadWrite, EditAnywhere, Category = "Draw Debug Library")
+	FString XaxisLabel;
+
+	/** Label for the Y axis */
+	UPROPERTY(BlueprintReadWrite, EditAnywhere, Category = "Draw Debug Library")
+	FString YaxisLabel;
+
+	/** Settings for the axis labels */
+	UPROPERTY(BlueprintReadWrite, EditAnywhere, Category = "Draw Debug Library")
+	FDrawDebugStringSettings AxisLabelSettings;
+};
+
 /**
  * A simple Blueprint wrapper around different objects that are capable of debug drawing. This allows us to use the same blueprint interface and
  * functions for debug drawing if we are drawing using a UWorld object or a FPrimitiveDrawInterface object (for example). Since this contains raw 
  * pointers it should not be stored anywhere as these could become stale. It should always be treated as a temporary object.
  */
-USTRUCT(BlueprintType)
+USTRUCT(BlueprintType, Category = "Draw Debug Library")
 struct FDebugDrawer
 {
 	GENERATED_BODY()
 
-public:
+public: // Constructors
+
+	/** Make a null debug drawer which ignores draw commands */
+	UE_API FDebugDrawer();
+
+	/** Make a debug drawer for a UObject */
+	UE_API FDebugDrawer(UObject* InObject);
+
+	/** Make a debug drawer for a UWorld */
+	UE_API FDebugDrawer(UWorld* InWorld);
+
+	/** Make a debug drawer for an AnimInstanceProxy */
+	UE_API FDebugDrawer(FAnimInstanceProxy* InAnimInstanceProxy);
+
+	/** Make a debug drawer for a FPrimitiveDrawInterface */
+	UE_API FDebugDrawer(FPrimitiveDrawInterface* InPrimitiveDrawInterface);
+
+	/** Make a null debug drawer which ignores draw commands */
+	static UE_API FDebugDrawer MakeDebugDrawer();
+
+	/** Make a debug drawer for a UObject */
+	static UE_API FDebugDrawer MakeDebugDrawer(UObject* Object);
+
+	/** Make a debug drawer for a UWorld */
+	static UE_API FDebugDrawer MakeDebugDrawer(UWorld* World);
+
+	/** Make a debug drawer for an AnimInstanceProxy */
+	static UE_API FDebugDrawer MakeDebugDrawer(FAnimInstanceProxy* AnimInstanceProxy);
+
+	/** Make a debug drawer for a FPrimitiveDrawInterface */
+	static UE_API FDebugDrawer MakeDebugDrawer(FPrimitiveDrawInterface* PrimitiveDrawInterface);
+
+	/** Make a debug drawer for the Visual Logger */
+	static UE_API FDebugDrawer MakeVisualLoggerDebugDrawer(UObject* Object, const FName Category = TEXT("LogDrawDebugLibrary"), const EDrawDebugLogVerbosity Verbosity = EDrawDebugLogVerbosity::Display, const bool bDrawToScene = true);
+
+	/** Make a debug drawer for the Visual Logger */
+	static UE_API FDebugDrawer MakeVisualLoggerDebugDrawer(UObject* Object, const FLogCategoryBase& Category, const ELogVerbosity::Type Verbosity = ELogVerbosity::Display, const bool bDrawToScene = true);
+
+	/** Make a debug drawer by merging multiple other debug drawers of different types */
+	static UE_API FDebugDrawer MakeMergedDebugDrawer(const TArrayView<const FDebugDrawer> DebugDrawers);
+
+public: // Drawer specific interface
+
+	UE_API void DrawPoints(const TArrayView<const FVector> Locations, const FLinearColor& Color, const float Thickness, const bool bDepthTest) const;
+	UE_API void DrawLines(const TArrayView<const FVector> StartLocations, const TArrayView<const FVector> EndLocations, const FLinearColor& Color, const float Thickness, const bool bDepthTest) const;
+	UE_API void DrawLinesPairs(const TArrayView<const TPair<FVector, FVector>> Segments, const FLinearColor& Color, const float Thickness, const bool bDepthTest) const;
+
+	UE_API void VisualLoggerLogString(const FStringView String, const FLinearColor& Color) const;
+	UE_API void VisualLoggerDrawString(const FStringView String, const FVector& Location, const FLinearColor& Color) const;
+
+private:
 
 	/** Primitive Draw Interface - used for drawing in editor viewports. */
 	FPrimitiveDrawInterface* PDI = nullptr;
 
 	/** World object - used for normal debug drawing. */
-	UWorld* World = nullptr;
+	UPROPERTY()
+	TObjectPtr<UWorld> World = nullptr;
 
 	/** Anim Instance Proxy Object - used for debug drawing in an AnimNode. */
 	FAnimInstanceProxy* AnimInstanceProxy = nullptr;
 
 	/** Visual Logger Object - used for writing to the visual logger. */
-	UObject* VisualLoggerObject = nullptr;
+	UPROPERTY()
+	TObjectPtr<UObject> VisualLoggerObject = nullptr;
 
 	/** Visual Logger Category. */
 	FName VisualLoggerCategory = NAME_None;
@@ -93,14 +364,16 @@ public:
 	ELogVerbosity::Type VisualLoggerVerbosity = ELogVerbosity::Display;
 };
 
+
 /**
- * A blueprint library of additional debug draw functions that unifies drawing between the PDI and World Debug Drawing functions.
- *
- * The built-in world debug draw functions are only accessible via blueprints for actor types or other classes which have an implicit World associated
- * with them. Here we provide blueprint-accessible debug draw functions that are accessible from anywhere (as long as you have either a world or
- * PDI variable accessibly) and we also provide functions for debug drawing a number of additional types of things.
+ * A blueprint library of additional debug draw functions that unifies drawing between different debug drawing interfaces. Also includes a number of 
+ * helper functions and other convenience functions that are useful for debug drawing such as functions for recording histories.
+ * 
+ * General parameter structure for these functions is:
+ * 
+ *	Required Inputs, Color, Thickness, bDepthTest, Settings
  */
-UCLASS(BlueprintType)
+UCLASS(BlueprintType, Category = "Draw Debug Library")
 class UDrawDebugLibrary : public UBlueprintFunctionLibrary
 {
 	GENERATED_BODY()
@@ -146,13 +419,13 @@ public: // Visual Debugger Logging
 
 	/** Log a string to the visual logger. Only does anything with a VisualLoggerDebugDrawer. */
 	UFUNCTION(BlueprintCallable, Category = "Draw Debug Library", meta = (DevelopmentOnly, AutoCreateRefTerm = "String"))
-	static UE_API void LogVisualLoggerString(const FDebugDrawer& Drawer, const FString& String, const FLinearColor Color = FLinearColor::Red);
-	static UE_API void LogVisualLoggerStringView(const FDebugDrawer& Drawer, const FStringView String, const FLinearColor Color = FLinearColor::Red);
+	static UE_API void VisualLoggerLogString(const FDebugDrawer& Drawer, const FString& String, const FLinearColor Color = FLinearColor::Red);
+	static UE_API void VisualLoggerLogStringView(const FDebugDrawer& Drawer, const FStringView String, const FLinearColor Color = FLinearColor::Red);
 
 	/** Log some text to the visual logger. Only does anything with a VisualLoggerDebugDrawer. */
 	UFUNCTION(BlueprintCallable, Category = "Draw Debug Library", meta = (DevelopmentOnly, AutoCreateRefTerm = "Text"))
-	static UE_API void LogVisualLoggerText(const FDebugDrawer& Drawer, const FText& Text, const FLinearColor Color = FLinearColor::Red);
-
+	static UE_API void VisualLoggerLogText(const FDebugDrawer& Drawer, const FText& Text, const FLinearColor Color = FLinearColor::Red);
+	
 public: // Convenience functions
 
 	/** Convenience function for making an array of linearly spaced float values */
@@ -171,12 +444,12 @@ public: // Point & Line
 
 	/** Debug Draw a point */
 	UFUNCTION(BlueprintCallable, Category = "Draw Debug Library", meta = (DevelopmentOnly, AutoCreateRefTerm = "Location"))
-	static UE_API void DrawDebugPoint(const FDebugDrawer& Drawer, const FVector& Location, const FLinearColor Color = FLinearColor::Red, const float Size = 0.0f, const bool bDepthTest = true);
+	static UE_API void DrawDebugPoint(const FDebugDrawer& Drawer, const FVector& Location, const FLinearColor Color = FLinearColor::Red, const float Thickness = 0.0f, const bool bDepthTest = true);
 
 	/** Debug Draw points. These should be preferred over DrawDebugPoint where possible as it will batch drawing when required such as when using the visual logger. */
 	UFUNCTION(BlueprintCallable, Category = "Draw Debug Library", meta = (DevelopmentOnly))
-	static UE_API void DrawDebugPoints(const FDebugDrawer& Drawer, const TArray<FVector>& Locations, const FLinearColor Color = FLinearColor::Red, const float Size = 0.0f, const bool bDepthTest = true);
-	static UE_API void DrawDebugPointsArrayView(const FDebugDrawer& Drawer, const TArrayView<const FVector> Locations, const FLinearColor Color = FLinearColor::Red, const float Size = 0.0f, const bool bDepthTest = true);
+	static UE_API void DrawDebugPoints(const FDebugDrawer& Drawer, const TArray<FVector>& Locations, const FLinearColor Color = FLinearColor::Red, const float Thickness = 0.0f, const bool bDepthTest = true);
+	static UE_API void DrawDebugPointsArrayView(const FDebugDrawer& Drawer, const TArrayView<const FVector> Locations, const FLinearColor Color = FLinearColor::Red, const float Thickness = 0.0f, const bool bDepthTest = true);
 
 	/** Debug Draw a line */
 	UFUNCTION(BlueprintCallable, Category = "Draw Debug Library", meta = (DevelopmentOnly, AutoCreateRefTerm = "StartLocation, EndLocation"))
@@ -192,151 +465,195 @@ public: // Basic Shapes
 
 	/** Debug Draw a triangular base pyramid */
 	UFUNCTION(BlueprintCallable, Category = "Draw Debug Library", meta = (DevelopmentOnly, AutoCreateRefTerm = "Location, Rotation"))
-	static UE_API void DrawDebugTriangularBasePyramid(const FDebugDrawer& Drawer, const FVector& Location, const FRotator& Rotation, const FLinearColor Color = FLinearColor::Red, const float Length = 20.0f, const float Width = 20.0f, const float Thickness = 0.0f, const bool bDepthTest = true);
+	static UE_API void DrawDebugTriangularBasePyramid(const FDebugDrawer& Drawer, const FVector& Location, const FRotator& Rotation, const FLinearColor Color = FLinearColor::Red, const float Thickness = 0.0f, const bool bDepthTest = true, const float Length = 20.0f, const float Width = 20.0f);
 
 	/** Debug Draw a square base pyramid */
 	UFUNCTION(BlueprintCallable, Category = "Draw Debug Library", meta = (DevelopmentOnly, AutoCreateRefTerm = "Location, Rotation"))
-	static UE_API void DrawDebugSquareBasePyramid(const FDebugDrawer& Drawer, const FVector& Location, const FRotator& Rotation, const FLinearColor Color = FLinearColor::Red, const float Length = 20.0f, const float Width = 20.0f, const float Thickness = 0.0f, const bool bDepthTest = true);
+	static UE_API void DrawDebugSquareBasePyramid(const FDebugDrawer& Drawer, const FVector& Location, const FRotator& Rotation, const FLinearColor Color = FLinearColor::Red, const float Thickness = 0.0f, const bool bDepthTest = true, const float Length = 20.0f, const float Width = 20.0f);
 
 	/** Debug Draw a cone */
 	UFUNCTION(BlueprintCallable, Category = "Draw Debug Library", meta = (DevelopmentOnly, AutoCreateRefTerm = "Location, Rotation"))
-	static UE_API void DrawDebugCone(const FDebugDrawer& Drawer, const FVector& Location, const FRotator& Rotation, const FLinearColor Color = FLinearColor::Red, const float Length = 20.0f, const float Radius = 10.0f, const int32 Segments = 9, const float Thickness = 0.0f, const bool bDepthTest = true);
+	static UE_API void DrawDebugCone(const FDebugDrawer& Drawer, const FVector& Location, const FRotator& Rotation, const FLinearColor Color = FLinearColor::Red, const float Thickness = 0.0f, const bool bDepthTest = true, const float Length = 20.0f, const float Radius = 10.0f, const int32 Segments = 9);
 
 	/** Debug Draw an arc around the XY axis */
 	UFUNCTION(BlueprintCallable, Category = "Draw Debug Library", meta = (DevelopmentOnly, AutoCreateRefTerm = "Location, Rotation"))
-	static UE_API void DrawDebugArc(const FDebugDrawer& Drawer, const FVector& Location, const FRotator& Rotation, const float Angle = 360.0f, const FLinearColor Color = FLinearColor::Red, const float Radius = 10.0f, const float Thickness = 0.0f, const int32 Segments = 17, const bool bDepthTest = true);
+	static UE_API void DrawDebugArc(const FDebugDrawer& Drawer, const FVector& Location, const FRotator& Rotation, const float Angle = 360.0f, const FLinearColor Color = FLinearColor::Red, const float Thickness = 0.0f, const bool bDepthTest = true, const float Radius = 10.0f, const int32 Segments = 17);
 
 	/** Debug Draw an circle around the XY axis */
 	UFUNCTION(BlueprintCallable, Category = "Draw Debug Library", meta = (DevelopmentOnly, AutoCreateRefTerm = "Location, Rotation"))
-	static UE_API void DrawDebugCircle(const FDebugDrawer& Drawer, const FVector& Location, const FRotator& Rotation, const FLinearColor Color = FLinearColor::Red, const float Radius = 10.0f, const float Thickness = 0.0f, const int32 Segments = 17, const bool bDepthTest = true);
+	static UE_API void DrawDebugCircle(const FDebugDrawer& Drawer, const FVector& Location, const FRotator& Rotation, const FLinearColor Color = FLinearColor::Red, const float Thickness = 0.0f, const bool bDepthTest = true, const float Radius = 10.0f, const int32 Segments = 17);
+
+	/** Debug Draw an tick on a circle at a given angle */
+	UFUNCTION(BlueprintCallable, Category = "Draw Debug Library", meta = (DevelopmentOnly, AutoCreateRefTerm = "Location, Rotation"))
+	static UE_API void DrawDebugCircleTick(const FDebugDrawer& Drawer, const FVector& Location, const FRotator& Rotation, const float Angle, const FLinearColor Color = FLinearColor::Red, const float Thickness = 0.0f, const bool bDepthTest = true, const float Radius = 10.0f, const float Length = 2.5f, bool bInside = true);
+
+	/** Debug Draw ticks on a circle at the given angles */
+	UFUNCTION(BlueprintCallable, Category = "Draw Debug Library", meta = (DevelopmentOnly, AutoCreateRefTerm = "Location, Rotation, Angles"))
+	static UE_API void DrawDebugCircleTicks(const FDebugDrawer& Drawer, const FVector& Location, const FRotator& Rotation, const TArray<float>& Angles, const FLinearColor Color = FLinearColor::Red, const float Thickness = 0.0f, const bool bDepthTest = true, const float Radius = 10.0f, const float Length = 5.0f, bool bInside = true);
+	static UE_API void DrawDebugCircleTicksArrayView(const FDebugDrawer& Drawer, const FVector& Location, const FRotator& Rotation, const TArrayView<const float> Angles, const FLinearColor Color = FLinearColor::Red, const float Thickness = 0.0f, const bool bDepthTest = true, const float Radius = 10.0f, const float Length = 5.0f, bool bInside = true);
 
 	/** Debug Draw the outline of a circle (two circles) around the XY axis */
 	UFUNCTION(BlueprintCallable, Category = "Draw Debug Library", meta = (DevelopmentOnly, AutoCreateRefTerm = "Location, Rotation"))
-	static UE_API void DrawDebugCircleOutline(const FDebugDrawer& Drawer, const FVector& Location, const FRotator& Rotation, const FLinearColor Color = FLinearColor::Red, const float InnerRadius = 10.0f, const float OuterRadius = 15.0f, const float Thickness = 0.0f, const int32 Segments = 17, const bool bDepthTest = true);
+	static UE_API void DrawDebugCircleOutline(const FDebugDrawer& Drawer, const FVector& Location, const FRotator& Rotation, const FLinearColor Color = FLinearColor::Red, const float Thickness = 0.0f, const bool bDepthTest = true, const float InnerRadius = 10.0f, const float OuterRadius = 15.0f, const int32 Segments = 17);
+
+	/** Debug Draw a triangle on the XY axis */
+	UFUNCTION(BlueprintCallable, Category = "Draw Debug Library", meta = (DevelopmentOnly, AutoCreateRefTerm = "Location, Rotation"))
+	static UE_API void DrawDebugTriangle(const FDebugDrawer& Drawer, const FVector& Location, const FRotator& Rotation, const FLinearColor Color = FLinearColor::Red, const float Thickness = 0.0f, const bool bDepthTest = true, const float Radius = 10.0f);
+
+	/** Debug Draw a square on the XY axis */
+	UFUNCTION(BlueprintCallable, Category = "Draw Debug Library", meta = (DevelopmentOnly, AutoCreateRefTerm = "Location, Rotation"))
+	static UE_API void DrawDebugSquare(const FDebugDrawer& Drawer, const FVector& Location, const FRotator& Rotation, const FLinearColor Color = FLinearColor::Red, const float Thickness = 0.0f, const bool bDepthTest = true, const float HalfLength = 10.0f);
+
+	/** Debug Draw a diamond on the XY axis */
+	UFUNCTION(BlueprintCallable, Category = "Draw Debug Library", meta = (DevelopmentOnly, AutoCreateRefTerm = "Location, Rotation"))
+	static UE_API void DrawDebugDiamond(const FDebugDrawer& Drawer, const FVector& Location, const FRotator& Rotation, const FLinearColor Color = FLinearColor::Red, const float Thickness = 0.0f, const bool bDepthTest = true, const float Radius = 10.0f);
+
+	/** Debug Draw a pentagon on the XY axis */
+	UFUNCTION(BlueprintCallable, Category = "Draw Debug Library", meta = (DevelopmentOnly, AutoCreateRefTerm = "Location, Rotation"))
+	static UE_API void DrawDebugPentagon(const FDebugDrawer& Drawer, const FVector& Location, const FRotator& Rotation, const FLinearColor Color = FLinearColor::Red, const float Thickness = 0.0f, const bool bDepthTest = true, const float Radius = 10.0f);
+
+	/** Debug Draw a hexagon on the XY axis */
+	UFUNCTION(BlueprintCallable, Category = "Draw Debug Library", meta = (DevelopmentOnly, AutoCreateRefTerm = "Location, Rotation"))
+	static UE_API void DrawDebugHexagon(const FDebugDrawer& Drawer, const FVector& Location, const FRotator& Rotation, const FLinearColor Color = FLinearColor::Red, const float Thickness = 0.0f, const bool bDepthTest = true, const float Radius = 10.0f);
+
+	/** Debug Draw a heptagon on the XY axis */
+	UFUNCTION(BlueprintCallable, Category = "Draw Debug Library", meta = (DevelopmentOnly, AutoCreateRefTerm = "Location, Rotation"))
+	static UE_API void DrawDebugHeptagon(const FDebugDrawer& Drawer, const FVector& Location, const FRotator& Rotation, const FLinearColor Color = FLinearColor::Red, const float Thickness = 0.0f, const bool bDepthTest = true, const float Radius = 10.0f);
+
+	/** Debug Draw a octagon on the XY axis */
+	UFUNCTION(BlueprintCallable, Category = "Draw Debug Library", meta = (DevelopmentOnly, AutoCreateRefTerm = "Location, Rotation"))
+	static UE_API void DrawDebugOctagon(const FDebugDrawer& Drawer, const FVector& Location, const FRotator& Rotation, const FLinearColor Color = FLinearColor::Red, const float Thickness = 0.0f, const bool bDepthTest = true, const float Radius = 10.0f);
+
+	/** Debug Draw a nonagon on the XY axis */
+	UFUNCTION(BlueprintCallable, Category = "Draw Debug Library", meta = (DevelopmentOnly, AutoCreateRefTerm = "Location, Rotation"))
+	static UE_API void DrawDebugNonagon(const FDebugDrawer& Drawer, const FVector& Location, const FRotator& Rotation, const FLinearColor Color = FLinearColor::Red, const float Thickness = 0.0f, const bool bDepthTest = true, const float Radius = 10.0f);
+
+	/** Debug Draw a decagon on the XY axis */
+	UFUNCTION(BlueprintCallable, Category = "Draw Debug Library", meta = (DevelopmentOnly, AutoCreateRefTerm = "Location, Rotation"))
+	static UE_API void DrawDebugDecagon(const FDebugDrawer& Drawer, const FVector& Location, const FRotator& Rotation, const FLinearColor Color = FLinearColor::Red, const float Thickness = 0.0f, const bool bDepthTest = true, const float Radius = 10.0f);
+
+	/** Debug Draw a flat regular polygon on the XY axis */
+	UFUNCTION(BlueprintCallable, Category = "Draw Debug Library", meta = (DevelopmentOnly, AutoCreateRefTerm = "Location, Rotation"))
+	static UE_API void DrawDebugRegularPolygon(const FDebugDrawer& Drawer, const FVector& Location, const FRotator& Rotation, const FLinearColor Color = FLinearColor::Red, const float Thickness = 0.0f, const bool bDepthTest = true, const float Radius = 10.0f, const int32 Sides = 11);
 
 	/** Debug Draw an axis-aligned cross */
 	UFUNCTION(BlueprintCallable, Category = "Draw Debug Library", meta = (DevelopmentOnly, AutoCreateRefTerm = "Location, Rotation"))
-	static UE_API void DrawDebugLocator(const FDebugDrawer& Drawer, const FVector& Location, const FRotator& Rotation, const FLinearColor Color = FLinearColor::Red, const float DrawRadius = 10.0f, const float Thickness = 0.0f, const bool bDepthTest = true);
+	static UE_API void DrawDebugLocator(const FDebugDrawer& Drawer, const FVector& Location, const FRotator& Rotation, const FLinearColor Color = FLinearColor::Red, const float Thickness = 0.0f, const bool bDepthTest = true, const float Radius = 10.0f);
 
 	/** Debug Draw an axis-aligned cross rotated by 45 degrees */
 	UFUNCTION(BlueprintCallable, Category = "Draw Debug Library", meta = (DevelopmentOnly, AutoCreateRefTerm = "Location, Rotation"))
-	static UE_API void DrawDebugCrossLocator(const FDebugDrawer& Drawer, const FVector& Location, const FRotator& Rotation, const FLinearColor Color = FLinearColor::Red, const float DrawRadius = 10.0f, const float Thickness = 0.0f, const bool bDepthTest = true);
+	static UE_API void DrawDebugCrossLocator(const FDebugDrawer& Drawer, const FVector& Location, const FRotator& Rotation, const FLinearColor Color = FLinearColor::Red, const float Thickness = 0.0f, const bool bDepthTest = true, const float Radius = 10.0f);
 
 	/** Debug Draw an oriented box */
 	UFUNCTION(BlueprintCallable, Category = "Draw Debug Library", meta = (DevelopmentOnly, AutoCreateRefTerm = "Location, Rotation"))
-	static UE_API void DrawDebugBox(const FDebugDrawer& Drawer, const FVector& Location, const FRotator& Rotation, const FVector HalfExtents = FVector(10.0f, 10.0f, 10.0f), const FLinearColor Color = FLinearColor::Red, const float Thickness = 0.0f, const bool bDepthTest = true);
+	static UE_API void DrawDebugBox(const FDebugDrawer& Drawer, const FVector& Location, const FRotator& Rotation, const FLinearColor Color = FLinearColor::Red, const float Thickness = 0.0f, const bool bDepthTest = true, const FVector HalfExtents = FVector(10.0f, 10.0f, 10.0f));
 
 	/** Debug Draw a sphere */
 	UFUNCTION(BlueprintCallable, Category = "Draw Debug Library", meta = (DevelopmentOnly, AutoCreateRefTerm = "Location, Rotation"))
-	static UE_API void DrawDebugSphere(const FDebugDrawer& Drawer, const FVector& Location, const FRotator& Rotation, const FLinearColor Color = FLinearColor::Red, const float Radius = 10.0f, const float Thickness = 0.0f, const int32 Segments = 8, const bool bDepthTest = true);
+	static UE_API void DrawDebugSphere(const FDebugDrawer& Drawer, const FVector& Location, const FRotator& Rotation, const FLinearColor Color = FLinearColor::Red, const float Thickness = 0.0f, const bool bDepthTest = true, const float Radius = 10.0f, const int32 Segments = 8);
+
+	/** Debug Draw a simple sphere made up of three circles one of each axis */
+	UFUNCTION(BlueprintCallable, Category = "Draw Debug Library", meta = (DevelopmentOnly, AutoCreateRefTerm = "Location, Rotation"))
+	static UE_API void DrawDebugSimpleSphere(const FDebugDrawer& Drawer, const FVector& Location, const FRotator& Rotation, const FLinearColor Color = FLinearColor::Red, const float Thickness = 0.0f, const bool bDepthTest = true, const float Radius = 10.0f, const int32 Segments = 13);
 
 	/** Debug Draw a capsule */
 	UFUNCTION(BlueprintCallable, Category = "Draw Debug Library", meta = (DevelopmentOnly, AutoCreateRefTerm = "Location, Rotation"))
-	static UE_API void DrawDebugCapsule(const FDebugDrawer& Drawer, const FVector& Location, const FRotator& Rotation, const FLinearColor Color = FLinearColor::Red, const float Radius = 10.0f, const float HalfLength = 10.0f, const float Thickness = 0.0f, const int32 Segments = 8, const bool bDepthTest = true);
+	static UE_API void DrawDebugCapsule(const FDebugDrawer& Drawer, const FVector& Location, const FRotator& Rotation, const FLinearColor Color = FLinearColor::Red, const float Thickness = 0.0f, const bool bDepthTest = true, const float Radius = 10.0f, const float HalfLength = 10.0f, const int32 Segments = 8);
+
+	/** Debug Draw a capsule using the start and end location */
+	UFUNCTION(BlueprintCallable, Category = "Draw Debug Library", meta = (DevelopmentOnly, AutoCreateRefTerm = "StartLocation, EndLocation"))
+	static UE_API void DrawDebugCapsuleLine(const FDebugDrawer& Drawer, const FVector& StartLocation, const FVector& EndLocation, const FLinearColor Color = FLinearColor::Red, const float Thickness = 0.0f, const bool bDepthTest = true, const float Radius = 10.0f, const int32 Segments = 8);
 
 	/** Debug Draw a Frustum */
 	UFUNCTION(BlueprintCallable, Category = "Draw Debug Library", meta = (DevelopmentOnly, AutoCreateRefTerm = "FrustumToWorld"))
 	static UE_API void DrawDebugFrustum(const FDebugDrawer& Drawer, const FMatrix& FrustumToWorld, const FLinearColor Color = FLinearColor::Red, const float Thickness = 0.0f, const bool bDepthTest = true);
 
-
 public: // Arrows
 
 	/** Debug Draw an arrow */
-	UFUNCTION(BlueprintCallable, Category = "Draw Debug Library", meta = (DevelopmentOnly, AutoCreateRefTerm = "StartLocation, EndLocation"))
-	static UE_API void DrawDebugArrow(const FDebugDrawer& Drawer, const FVector& StartLocation, const FVector& EndLocation, const FLinearColor Color = FLinearColor::Red, const float ArrowHeadSize = 20.0f, const float Thickness = 0.0f, const bool bDepthTest = true);
+	UFUNCTION(BlueprintCallable, Category = "Draw Debug Library", meta = (DevelopmentOnly, AutoCreateRefTerm = "StartLocation, EndLocation, Settings"))
+	static UE_API void DrawDebugArrow(const FDebugDrawer& Drawer, const FVector& StartLocation, const FVector& EndLocation, const FLinearColor Color = FLinearColor::Red, const float Thickness = 0.0f, const bool bDepthTest = true, const FDrawDebugArrowSettings& Settings = FDrawDebugArrowSettings());
 
 	/** Debug Draw an arrow with an orientation */
-	UFUNCTION(BlueprintCallable, Category = "Draw Debug Library", meta = (DevelopmentOnly, AutoCreateRefTerm = "Location, Rotation"))
-	static UE_API void DrawDebugOrientedArrow(const FDebugDrawer& Drawer, const FVector& Location, const FRotator& Rotation, const float Length = 100.0f, const FLinearColor Color = FLinearColor::Red, const float ArrowHeadSize = 20.0f, const float Thickness = 0.0f, const bool bDepthTest = true);
-
-	/** Debug Draw an arrow with a four-way head */
-	UFUNCTION(BlueprintCallable, Category = "Draw Debug Library", meta = (DevelopmentOnly, AutoCreateRefTerm = "StartLocation, EndLocation"))
-	static UE_API void DrawDebugFourWayArrow(const FDebugDrawer& Drawer, const FVector& StartLocation, const FVector& EndLocation, const FLinearColor Color = FLinearColor::Red, const float ArrowHeadSize = 20.0f, const float Thickness = 0.0f, const bool bDepthTest = true);
+	UFUNCTION(BlueprintCallable, Category = "Draw Debug Library", meta = (DevelopmentOnly, AutoCreateRefTerm = "Location, Rotation, Settings"))
+	static UE_API void DrawDebugOrientedArrow(const FDebugDrawer& Drawer, const FVector& Location, const FRotator& Rotation, const float Length = 100.0f, const FLinearColor Color = FLinearColor::Red, const float Thickness = 0.0f, const bool bDepthTest = true, const FDrawDebugArrowSettings& Settings = FDrawDebugArrowSettings());
 
 	/** Draws an arrow pointing down at the location surrounded by a circle */
 	UFUNCTION(BlueprintCallable, Category = "Draw Debug Library", meta = (DevelopmentOnly, AutoCreateRefTerm = "Location"))
-	static UE_API void DrawDebugGroundTargetArrow(const FDebugDrawer& Drawer, const FVector& Location, const FLinearColor Color = FLinearColor::Red, const float Length = 100.0f, const float ArrowHeadSize = 20.0f, const float Radius = 10.0f, const int32 Segements = 17, const float Thickness = 0.0f, const bool bDepthTest = true);
-
-	/** Debug Draw an arrow with a triangular arrowhead */
-	UFUNCTION(BlueprintCallable, Category = "Draw Debug Library", meta = (DevelopmentOnly, AutoCreateRefTerm = "StartLocation, EndLocation"))
-	static UE_API void DrawTriangularArrow(const FDebugDrawer& Drawer, const FVector& StartLocation, const FVector& EndLocation, const FLinearColor Color = FLinearColor::Red, const float ArrowHeadSize = 20.0f, const float Thickness = 0.0f, const bool bDepthTest = true);
-
-	/** Debug Draw an arrow with a square arrowhead */
-	UFUNCTION(BlueprintCallable, Category = "Draw Debug Library", meta = (DevelopmentOnly, AutoCreateRefTerm = "StartLocation, EndLocation"))
-	static UE_API void DrawSquareArrow(const FDebugDrawer& Drawer, const FVector& StartLocation, const FVector& EndLocation, const FLinearColor Color = FLinearColor::Red, const float ArrowHeadSize = 20.0f, const float Thickness = 0.0f, const bool bDepthTest = true);
-
-	/** Debug Draw an arrow with a cone arrowhead */
-	UFUNCTION(BlueprintCallable, Category = "Draw Debug Library", meta = (DevelopmentOnly, AutoCreateRefTerm = "StartLocation, EndLocation"))
-	static UE_API void DrawConeArrow(const FDebugDrawer& Drawer, const FVector& StartLocation, const FVector& EndLocation, const FLinearColor Color = FLinearColor::Red, const float ArrowHeadSize = 20.0f, const float Thickness = 0.0f, const int32 Segments = 9, const bool bDepthTest = true);
+	static UE_API void DrawDebugGroundTargetArrow(const FDebugDrawer& Drawer, const FVector& Location, const FLinearColor Color = FLinearColor::Red, const float Thickness = 0.0f, const bool bDepthTest = true, const float Length = 100.0f, const float ArrowHeadSize = 20.0f, const float Radius = 10.0f, const int32 Segments = 17);
 
 	/** Draws a flat 2d arrow motif */
 	UFUNCTION(BlueprintCallable, Category = "Draw Debug Library", meta = (DevelopmentOnly, AutoCreateRefTerm = "Location, Rotation"))
-	static UE_API void DrawDebugFlatArrow(const FDebugDrawer& Drawer, const FVector& Location, const FRotator& Rotation, const FLinearColor Color = FLinearColor::Red, const float Length = 20.0f, const float Width = 10.0f, const float Thickness = 0.0f, const bool bDepthTest = true);
+	static UE_API void DrawDebugFlatArrow(const FDebugDrawer& Drawer, const FVector& Location, const FRotator& Rotation, const FLinearColor Color = FLinearColor::Red, const float Thickness = 0.0f, const bool bDepthTest = true, const float Length = 20.0f, const float Width = 20.0f);
+
+	/** Debug Draw an arrow coming off a circle at a given angle */
+	UFUNCTION(BlueprintCallable, Category = "Draw Debug Library", meta = (DevelopmentOnly, AutoCreateRefTerm = "Location, Rotation, ArrowSettings"))
+	static UE_API void DrawDebugCircleArrow(const FDebugDrawer& Drawer, const FVector& Location, const FRotator& Rotation, const float Angle, const FLinearColor Color = FLinearColor::Red, const float Thickness = 0.0f, const bool bDepthTest = true, const float Radius = 10.0f, const float Length = 20.0f, const FDrawDebugArrowSettings& ArrowSettings = FDrawDebugArrowSettings());
 
 public: // Splines
 
 	/** Debug Draw the start section of a catmull rom spline */
 	UFUNCTION(BlueprintCallable, Category = "Draw Debug Library", meta = (DevelopmentOnly, AutoCreateRefTerm = "V0, V1, V2"))
-	static UE_API void DrawDebugCatmullRomSplineStart(const FDebugDrawer& Drawer, const FVector& V0, const FVector& V1, const FVector& V2, const bool bMonotonic = false, const FLinearColor Color = FLinearColor::Red, const float Thickness = 0.0f, const int32 Segments = 15, const bool bDepthTest = true);
+	static UE_API void DrawDebugCatmullRomSplineStart(const FDebugDrawer& Drawer, const FVector& V0, const FVector& V1, const FVector& V2, const FLinearColor Color = FLinearColor::Red, const float Thickness = 0.0f, const bool bDepthTest = true, const bool bMonotonic = false, const int32 Segments = 15);
 
 	/** Debug Draw the end section of a catmull rom spline */
 	UFUNCTION(BlueprintCallable, Category = "Draw Debug Library", meta = (DevelopmentOnly, AutoCreateRefTerm = "V0, V1, V2"))
-	static UE_API void DrawDebugCatmullRomSplineEnd(const FDebugDrawer& Drawer, const FVector& V0, const FVector& V1, const FVector& V2, const bool bMonotonic = false, const FLinearColor Color = FLinearColor::Red, const float Thickness = 0.0f, const int32 Segments = 15, const bool bDepthTest = true);
+	static UE_API void DrawDebugCatmullRomSplineEnd(const FDebugDrawer& Drawer, const FVector& V0, const FVector& V1, const FVector& V2, const FLinearColor Color = FLinearColor::Red, const float Thickness = 0.0f, const bool bDepthTest = true, const bool bMonotonic = false, const int32 Segments = 15);
 
 	/** Debug Draw a full segment of a catmull rom spline */
 	UFUNCTION(BlueprintCallable, Category = "Draw Debug Library", meta = (DevelopmentOnly, AutoCreateRefTerm = "V0, V1, V2, V3"))
-	static UE_API void DrawDebugCatmullRomSplineSection(const FDebugDrawer& Drawer, const FVector& V0, const FVector& V1, const FVector& V2, const FVector& V3, const bool bMonotonic = false, const FLinearColor Color = FLinearColor::Red, const float Thickness = 0.0f, const int32 Segments = 15, const bool bDepthTest = true);
+	static UE_API void DrawDebugCatmullRomSplineSection(const FDebugDrawer& Drawer, const FVector& V0, const FVector& V1, const FVector& V2, const FVector& V3, const FLinearColor Color = FLinearColor::Red, const float Thickness = 0.0f, const bool bDepthTest = true, const bool bMonotonic = false, const int32 Segments = 15);
 
 	/** Debug Draw a catmull rom spline */
 	UFUNCTION(BlueprintCallable, Category = "Draw Debug Library", meta = (DevelopmentOnly))
-	static UE_API void DrawDebugCatmullRomSpline(const FDebugDrawer& Drawer, const TArray<FVector>& Points, const bool bMonotonic = false, const FLinearColor Color = FLinearColor::Red, const float Thickness = 0.0f, const int32 Segments = 15, const bool bDepthTest = true);
-	static UE_API void DrawDebugCatmullRomSplineArrayView(const FDebugDrawer& Drawer, const TArrayView<const FVector> Points, const bool bMonotonic = false, const FLinearColor Color = FLinearColor::Red, const float Thickness = 0.0f, const int32 Segments = 15, const bool bDepthTest = true);
+	static UE_API void DrawDebugCatmullRomSpline(const FDebugDrawer& Drawer, const TArray<FVector>& Points, const FLinearColor Color = FLinearColor::Red, const float Thickness = 0.0f, const bool bDepthTest = true, const bool bMonotonic = false, const int32 Segments = 15);
+	static UE_API void DrawDebugCatmullRomSplineArrayView(const FDebugDrawer& Drawer, const TArrayView<const FVector> Points, const FLinearColor Color = FLinearColor::Red, const float Thickness = 0.0f, const bool bDepthTest = true, const bool bMonotonic = false, const int32 Segments = 15);
 
 public: // Variable Types
 
 	/** Debug Draw an angle */
 	UFUNCTION(BlueprintCallable, Category = "Draw Debug Library", meta = (DevelopmentOnly, AutoCreateRefTerm = "Location, Rotation"))
-	static UE_API void DrawDebugAngle(const FDebugDrawer& Drawer, const FVector& Location, const FRotator& Rotation, const float Angle = 0.0f, const float LineLength = 10.0f, const float AngleRadius = 5.0f, const FLinearColor Color = FLinearColor::Red, const float Thickness = 0.0f, const int32 Segments = 17, const bool bDepthTest = true);
+	static UE_API void DrawDebugAngle(const FDebugDrawer& Drawer, const FVector& Location, const FRotator& Rotation, const float Angle = 0.0f, const FLinearColor Color = FLinearColor::Red, const float Thickness = 0.0f, const bool bDepthTest = true, const float LineLength = 10.0f, const float AngleRadius = 5.0f, const int32 Segments = 17);
 
 	/** Debug Draw a sphere at the given location */
 	UFUNCTION(BlueprintCallable, Category = "Draw Debug Library", meta = (DevelopmentOnly, AutoCreateRefTerm = "Location"))
-	static UE_API void DrawDebugLocation(const FDebugDrawer& Drawer, const FVector& Location, const FLinearColor Color = FLinearColor::Red, const float DrawRadius = 10.0f, const float Thickness = 0.0f, const int32 Segments = 8, const bool bDepthTest = true);
+	static UE_API void DrawDebugLocation(const FDebugDrawer& Drawer, const FVector& Location, const FLinearColor Color = FLinearColor::Red, const float Thickness = 0.0f, const bool bDepthTest = true, const float DrawRadius = 10.0f, const int32 Segments = 8);
 
 	/** Debug Draw spheres at the given locations */
 	UFUNCTION(BlueprintCallable, Category = "Draw Debug Library", meta = (DevelopmentOnly))
-	static UE_API void DrawDebugLocations(const FDebugDrawer& Drawer, const TArray<FVector>& Locations, const FLinearColor Color = FLinearColor::Red, const float DrawRadius = 10.0f, const float Thickness = 0.0f, const int32 Segments = 8, const bool bDepthTest = true);
-	static UE_API void DrawDebugLocationsArrayView(const FDebugDrawer& Drawer, const TArrayView<const FVector> Locations, const FLinearColor Color = FLinearColor::Red, const float DrawRadius = 10.0f, const float Thickness = 0.0f, const int32 Segments = 8, const bool bDepthTest = true);
+	static UE_API void DrawDebugLocations(const FDebugDrawer& Drawer, const TArray<FVector>& Locations, const FLinearColor Color = FLinearColor::Red, const float Thickness = 0.0f, const bool bDepthTest = true, const float DrawRadius = 10.0f, const int32 Segments = 8);
+	static UE_API void DrawDebugLocationsArrayView(const FDebugDrawer& Drawer, const TArrayView<const FVector> Locations, const FLinearColor Color = FLinearColor::Red, const float Thickness = 0.0f, const bool bDepthTest = true, const float DrawRadius = 10.0f, const int32 Segments = 8);
 
 	/** Draws a rotation */
 	UFUNCTION(BlueprintCallable, Category = "Draw Debug Library", meta = (DevelopmentOnly, AutoCreateRefTerm = "Location, Rotation"))
-	static UE_API void DrawDebugRotation(const FDebugDrawer& Drawer, const FVector& Location, const FRotator& Rotation, const float DrawRadius = 10.0f, const float Thickness = 0.0f, const float Opacity = 1.0f, const bool bDepthTest = true);
+	static UE_API void DrawDebugRotation(const FDebugDrawer& Drawer, const FVector& Location, const FRotator& Rotation, const float Opacity = 1.0f, const float Thickness = 0.0f, const bool bDepthTest = true, const float DrawRadius = 10.0f);
 
 	/** Draws an array of transforms */
 	static UE_API void DrawDebugRotationsQuatArrayView(
 		const FDebugDrawer& Drawer,
 		const TArrayView<const FVector> Locations,
 		const TArrayView<const FQuat4f> Rotations,
-		const float DrawRadius = 10.0f,
-		const float Thickness = 0.0f,
 		const float Opacity = 1.0f,
-		const bool bDepthTest = true);
+		const float Thickness = 0.0f,
+		const bool bDepthTest = true,
+		const float DrawRadius = 10.0f);
 
 	/** Debug Draw an arrow at the given location facing in the given direction */
 	UFUNCTION(BlueprintCallable, Category = "Draw Debug Library", meta = (DevelopmentOnly, AutoCreateRefTerm = "Location, Direction"))
-	static UE_API void DrawDebugDirection(const FDebugDrawer& Drawer, const FVector& Location, const FVector& Direction = FVector::ForwardVector, const FLinearColor Color = FLinearColor::Red, const float DrawArrowLength = 100.0f, const float ArrowHeadScale = 1.0f, const float Thickness = 0.0f, const bool bDepthTest = true);
+	static UE_API void DrawDebugDirection(const FDebugDrawer& Drawer, const FVector& Location, const FVector& Direction = FVector::ForwardVector, const FLinearColor Color = FLinearColor::Red, const float Thickness = 0.0f, const bool bDepthTest = true, const float DrawArrowLength = 100.0f, const float ArrowHeadScale = 1.0f);
 
 	/** Debug Draw a line at the given location scaled by the given velocity */
 	UFUNCTION(BlueprintCallable, Category = "Draw Debug Library", meta = (DevelopmentOnly, AutoCreateRefTerm = "Location, Velocity"))
-	static UE_API void DrawDebugVelocity(const FDebugDrawer& Drawer, const FVector& Location, const FVector& Velocity, const FLinearColor Color = FLinearColor::Red, const float DrawVelocityLineScale = 1.0f, const float Thickness = 0.0f, const bool bDepthTest = true);
+	static UE_API void DrawDebugVelocity(const FDebugDrawer& Drawer, const FVector& Location, const FVector& Velocity, const FLinearColor Color = FLinearColor::Red, const float Thickness = 0.0f, const bool bDepthTest = true, const float DrawVelocityLineScale = 1.0f);
 
 	/** Debug draw velocities */
 	UFUNCTION(BlueprintCallable, Category = "Draw Debug Library", meta = (DevelopmentOnly))
-	static UE_API void DrawDebugVelocities(const FDebugDrawer& Drawer, const TArray<FVector>& Locations, const TArray<FVector>& Velocities, const FLinearColor Color = FLinearColor::Red, const float DrawVelocityLineScale = 1.0f, const float Thickness = 0.0f, const bool bDepthTest = true);
-	static UE_API void DrawDebugVelocitiesArrayView(const FDebugDrawer& Drawer, const TArrayView<const FVector> Locations, const TArrayView<const FVector> Velocities, const FLinearColor Color = FLinearColor::Red, const float DrawVelocityLineScale = 1.0f, const float Thickness = 0.0f, const bool bDepthTest = true);
+	static UE_API void DrawDebugVelocities(const FDebugDrawer& Drawer, const TArray<FVector>& Locations, const TArray<FVector>& Velocities, const FLinearColor Color = FLinearColor::Red, const float Thickness = 0.0f, const bool bDepthTest = true, const float DrawVelocityLineScale = 1.0f);
+	static UE_API void DrawDebugVelocitiesArrayView(const FDebugDrawer& Drawer, const TArrayView<const FVector> Locations, const TArrayView<const FVector> Velocities, const FLinearColor Color = FLinearColor::Red, const float Thickness = 0.0f, const bool bDepthTest = true, const float DrawVelocityLineScale = 1.0f);
 
 	/** Debug Draw a set of axes at the given transform */
 	UFUNCTION(BlueprintCallable, Category = "Draw Debug Library", meta = (DevelopmentOnly, AutoCreateRefTerm = "Transform"))
-	static UE_API void DrawDebugTransform(const FDebugDrawer& Drawer, const FTransform& Transform, const float DrawRadius = 10.0f, const float Thickness = 0.0f, const float Opacity = 1.0f, const bool bDepthTest = true);
+	static UE_API void DrawDebugTransform(const FDebugDrawer& Drawer, const FTransform& Transform, const float Opacity = 1.0f, const float Thickness = 0.0f, const bool bDepthTest = true, const float DrawRadius = 10.0f);
 
 	/** Draws an array of transforms */
 	static UE_API void DrawDebugTransformsArrayView(
@@ -344,81 +661,94 @@ public: // Variable Types
 		const TArrayView<const FVector> Locations,
 		const TArrayView<const FQuat4f> Rotations,
 		const TArrayView<const FVector3f> Scales,
-		const float DrawRadius = 10.0f,
-		const float Thickness = 0.0f,
 		const float Opacity = 1.0f,
-		const bool bDepthTest = true);
+		const float Thickness = 0.0f,
+		const bool bDepthTest = true,
+		const float DrawRadius = 10.0f);
 
 	/** Debug draw a phase-representation of an event at a location */
 	UFUNCTION(BlueprintCallable, Category = "Draw Debug Library", meta = (DevelopmentOnly, AutoCreateRefTerm = "Location"))
-	static UE_API void DrawDebugEvent(const FDebugDrawer& Drawer, bool bTimeUntilEventKnown, const float TimeUntilEvent, const FVector& Location, const FLinearColor Color = FLinearColor::Red, const float Size = 20.0f, const float Thickness = 0.0f, const bool bDepthTest = true);
+	static UE_API void DrawDebugEvent(const FDebugDrawer& Drawer, bool bTimeUntilEventKnown, const float TimeUntilEvent, const FVector& Location, const FLinearColor Color = FLinearColor::Red, const float Thickness = 0.0f, const bool bDepthTest = true, const float Size = 20.0f);
 
 public: // Complex Objects
 
 	/** Debug draw a basic chair shape */
-	UFUNCTION(BlueprintCallable, Category = "Draw Debug Library", meta = (DevelopmentOnly, AutoCreateRefTerm = "Location, Rotation"))
-	static UE_API void DrawDebugChair(const FDebugDrawer& Drawer, const FVector& Location, const FRotator& Rotation, const float Width = 40.0f, const float SeatHeight = 40.0f, const float BackHeight = 40.0f, const float BackTilt = 10.0f, const FLinearColor Color = FLinearColor::Red, const float Thickness = 0.0f, const bool bDepthTest = true);
+	UFUNCTION(BlueprintCallable, Category = "Draw Debug Library", meta = (DevelopmentOnly, AutoCreateRefTerm = "Location, Rotation, Settings"))
+	static UE_API void DrawDebugChair(const FDebugDrawer& Drawer, const FVector& Location, const FRotator& Rotation, const FLinearColor Color = FLinearColor::Red, const float Thickness = 0.0f, const bool bDepthTest = true, const FDrawDebugChairSettings& Settings = FDrawDebugChairSettings());
 
 	/** Debug draw a basic door shape */
-	UFUNCTION(BlueprintCallable, Category = "Draw Debug Library", meta = (DevelopmentOnly, AutoCreateRefTerm = "Location, Rotation"))
-	static UE_API void DrawDebugDoor(const FDebugDrawer& Drawer, const FVector& Location, const FRotator& Rotation, const float OpenAngle = 0.0f, const float Width = 100.0f, const float Height = 200.0f, const float Inset = 5.0f, const float HandleHeight = 90.0f, const float HandleOffset = 20.0f, const float HandleRadius = 3.0f, const float HandleExtension = 2.0f, const bool bOuterHandle = true, const bool bInnerHandle = true, const bool bHandleOnLeft = false, const bool bDrawEntryArrow = true, const FLinearColor Color = FLinearColor::Red, const float Thickness = 0.0f, const bool bDepthTest = true);
+	UFUNCTION(BlueprintCallable, Category = "Draw Debug Library", meta = (DevelopmentOnly, AutoCreateRefTerm = "Location, Rotation, Settings"))
+	static UE_API void DrawDebugDoor(const FDebugDrawer& Drawer, const FVector& Location, const FRotator& Rotation, const FLinearColor Color = FLinearColor::Red, const float Thickness = 0.0f, const bool bDepthTest = true, const FDrawDebugDoorSettings& Settings = FDrawDebugDoorSettings());
 
 	/** Debug draw a camera shape */
 	UFUNCTION(BlueprintCallable, Category = "Draw Debug Library", meta = (DevelopmentOnly, AutoCreateRefTerm = "Location, Rotation"))
-	static UE_API void DrawDebugCamera(const FDebugDrawer& Drawer, const FVector& Location, const FRotator& Rotation, float Scale = 10.0f, float FOVDeg = 30.0f, const FLinearColor Color = FLinearColor::Red, const float Thickness = 0.0f, const bool bDepthTest = true);
+	static UE_API void DrawDebugCamera(const FDebugDrawer& Drawer, const FVector& Location, const FRotator& Rotation, const FLinearColor Color = FLinearColor::Red, const float Thickness = 0.0f, const bool bDepthTest = true, const float Scale = 10.0f, const float FOVDegees = 30.0f);
 
 	/** Debug Draw a trajectory of locations and directions */
 	UFUNCTION(BlueprintCallable, Category = "Draw Debug Library", meta = (DevelopmentOnly, AutoCreateRefTerm = "RelativeTransform"))
-	static UE_API void DrawDebugTrajectory(const FDebugDrawer& Drawer, const TArray<FVector>& Locations, const TArray<FVector>& Directions, const FTransform& RelativeTransform = FTransform(), const FLinearColor Color = FLinearColor::Red, const float DrawArrowLength = 100.0f, const float PointRadius = 5.0f, const float ArrowHeadScale = 1.0f, const float Thickness = 0.0f, const int32 Segments = 9, const float VerticalOffset = 2.0f, const bool bDepthTest = true);
-	static UE_API void DrawDebugTrajectoryFromArrayViews(const FDebugDrawer& Drawer, const TArrayView<const FVector> Locations, const TArrayView<const FVector> Directions, const FTransform& RelativeTransform = FTransform(), const FLinearColor Color = FLinearColor::Red, const float DrawArrowLength = 100.0f, const float PointRadius = 5.0f, const float ArrowHeadScale = 1.0f, const float Thickness = 0.0f, const int32 Segments = 9, const float VerticalOffset = 2.0f, const bool bDepthTest = true);
+	static UE_API void DrawDebugTrajectory(const FDebugDrawer& Drawer, const TArray<FVector>& Locations, const TArray<FVector>& Directions, const FTransform& RelativeTransform = FTransform(), const FLinearColor Color = FLinearColor::Red, const float Thickness = 0.0f, const bool bDepthTest = true, const float DrawArrowLength = 100.0f, const float PointRadius = 5.0f, const float ArrowHeadScale = 1.0f, const int32 Segments = 9, const float VerticalOffset = 2.0f);
+	static UE_API void DrawDebugTrajectoryFromArrayViews(const FDebugDrawer& Drawer, const TArrayView<const FVector> Locations, const TArrayView<const FVector> Directions, const FTransform& RelativeTransform = FTransform(), const FLinearColor Color = FLinearColor::Red, const float Thickness = 0.0f, const bool bDepthTest = true, const float DrawArrowLength = 100.0f, const float PointRadius = 5.0f, const float ArrowHeadScale = 1.0f, const int32 Segments = 9, const float VerticalOffset = 2.0f);
 
 	/** Debug Draw a transform trajectory */
 	UFUNCTION(BlueprintCallable, Category = "Draw Debug Library", meta = (DevelopmentOnly))
-	static UE_API void DrawDebugTransformTrajectory(const FDebugDrawer& Drawer, UPARAM(ref) const FTransformTrajectory& TransformTrajectory, const FLinearColor Color = FLinearColor::Red, const float Radius = 5.0f, const float Thickness = 0.0f, const float VerticalOffset = 2.0f, const bool bDepthTest = true);
+	static UE_API void DrawDebugTransformTrajectory(const FDebugDrawer& Drawer, UPARAM(ref) const FTransformTrajectory& TransformTrajectory, const FLinearColor Color = FLinearColor::Red, const float Thickness = 0.0f, const bool bDepthTest = true, const float Radius = 5.0f, const float VerticalOffset = 2.0f);
 
 	/** Draws a trajectory in the style of trajectory drawing in persona */
 	static UE_API void DrawDebugRangeTrajectoryArrayView(
 		const FDebugDrawer& Drawer,
 		const TArrayView<const FVector> TrajectoryLocations,
 		const TArrayView<const FQuat4f> TrajectoryRotations,
-		const FVector& ForwardVector,
-		const bool bDrawOrientations,
-		const bool bDrawStartingAtOrigin,
-		const FVector& OriginOffset,
-		const bool bDepthTest = true);
+		const bool bDepthTest = true,
+		const FVector& ForwardVector = FVector::ForwardVector,
+		const bool bDrawOrientations = false,
+		const bool bDrawStartingAtOrigin = false,
+		const FVector& OriginOffset = FVector::ZeroVector);
 
 	/** Draws a flat 2d motif representing mover position and rotation */
 	UFUNCTION(BlueprintCallable, Category = "Draw Debug Library", meta = (DevelopmentOnly, AutoCreateRefTerm = "Location, Rotation"))
-	static UE_API void DrawDebugMoverOrientation(const FDebugDrawer& Drawer, const FVector& Location, const FRotator& Rotation, const FVector ForwardVector = FVector::ForwardVector, const float Scale = 30.0f, const FLinearColor Color = FLinearColor::Red, const float Thickness = 0.0f, const bool bDepthTest = true);
+	static UE_API void DrawDebugMoverOrientation(const FDebugDrawer& Drawer, const FVector& Location, const FRotator& Rotation, const FLinearColor Color = FLinearColor::Red, const float Thickness = 0.0f, const bool bDepthTest = true, const FVector ForwardVector = FVector::ForwardVector, const float Scale = 30.0f);
 
 public: // Poses / Skeletons
 
-	/** Debug Draw a pose made up of bone locations and velocities */
-	UFUNCTION(BlueprintCallable, Category = "Draw Debug Library", meta = (DevelopmentOnly, AutoCreateRefTerm = "RelativeTransform"))
-	static UE_API void DrawDebugPose(const FDebugDrawer& Drawer, const TArray<FVector>& BoneLocations, const TArray<FVector>& BoneLinearVelocities, const FTransform& RelativeTransform = FTransform(), const FLinearColor Color = FLinearColor::Red, const float DrawVelocityLineScale = 1.0f, const float Thickness = 0.0f, const bool bDepthTest = true);
-	static UE_API void DrawDebugPoseFromArrayViews(const FDebugDrawer& Drawer, const TArrayView<const FVector> BoneLocations, const TArrayView<const FVector> BoneLinearVelocities, const FTransform& RelativeTransform = FTransform(), const FLinearColor Color = FLinearColor::Red, const float DrawVelocityLineScale = 1.0f, const float Thickness = 0.0f, const bool bDepthTest = true);
+	/** Gets the default debug draw bone color */
+	UFUNCTION(BlueprintPure, Category = "Draw Debug Library")
+	static UE_API FLinearColor GetDefaultBoneColor();
+
+	/** Gets the default debug draw bone radius */
+	UFUNCTION(BlueprintPure, Category = "Draw Debug Library")
+	static UE_API float GetDefaultBoneRadius();
+
+	/** Debug Draw a bone */
+	UFUNCTION(BlueprintCallable, Category = "Draw Debug Library", meta = (DevelopmentOnly, AutoCreateRefTerm = "Location, Rotation"))
+	static UE_API void DrawDebugBone(const FDebugDrawer& Drawer, const FVector& Location, const FRotator& Rotation, const FLinearColor Color = FLinearColor(0.0f, 0.0f, 0.025f, 1.0f), const float Thickness = 0.0f, const bool bDepthTest = true, const float Radius = 5.0f, const int32 Segments = 10, const bool bDrawTransform = false);
+
+	/** Debug Draw a link between a child and parent bone */
+	UFUNCTION(BlueprintCallable, Category = "Draw Debug Library", meta = (DevelopmentOnly, AutoCreateRefTerm = "ChildLocation, ParentLocation, ParentRotation"))
+	static UE_API void DrawDebugBoneLink(const FDebugDrawer& Drawer, const FVector& ChildLocation, const FVector& ParentLocation, const FRotator& ParentRotation, const FLinearColor Color = FLinearColor(0.0f, 0.0f, 0.025f, 1.0f), const float Thickness = 0.0f, const bool bDepthTest = true, const float Radius = 5.0f);
+
+	/** Draws a skeleton from a skinned mesh component */
+	UFUNCTION(BlueprintCallable, Category = "Draw Debug Library", meta = (DevelopmentOnly))
+	static UE_API void DrawDebugSkeletonFromSkinnedMeshComponent(
+		const FDebugDrawer& Drawer,
+		const USkinnedMeshComponent* SkinnedMeshComponent,
+		const FLinearColor Color = FLinearColor(0.0f, 0.0f, 0.025f, 1.0f),
+		const float Thickness = 0.0f,
+		const bool bDepthTest = false,
+		const FDrawDebugSkeletonSettings& Settings = FDrawDebugSkeletonSettings());
 
 	/** Draws a skeleton made up of the given bone transforms */
 	static UE_API void DrawDebugSkeletonArrayView(
 		const FDebugDrawer& Drawer,
 		const FVector RootLocation,
+		const FQuat4f RootRotation,
 		const TArrayView<const FVector> BoneLocations,
 		const TArrayView<const FQuat4f> BoneRotations,
-		const TArrayView<const FVector3f> BoneScales,
-		const TArrayView<const int32> BoneIndices,
-		const FReferenceSkeleton& ReferenceSkeleton,
-		const FLinearColor& Color,
-		const float BoneDrawSize);
-
-	/** Draws a simple skeleton made up of the given bone locations */
-	static UE_API void DrawDebugSimpleSkeletonArrayView(
-		const FDebugDrawer& Drawer,
-		const TArrayView<const FVector> BoneLocations,
 		const TArrayView<const int32> BoneIndices,
 		const TArrayView<const int32> BoneParents,
-		const FLinearColor& Color,
+		const FLinearColor Color = FLinearColor(0.0f, 0.0f, 0.025f, 1.0f),
 		const float Thickness = 0.0f,
-		const bool bDepthTest = false);
+		const bool bDepthTest = false,
+		const FDrawDebugSkeletonSettings& Settings = FDrawDebugSkeletonSettings());
 
 	/** Draws lines representing bone velocities */
 	static UE_API void DrawDebugBoneVelocitiesArrayView(
@@ -426,77 +756,69 @@ public: // Poses / Skeletons
 		const TArrayView<const FVector> BoneLocations,
 		const TArrayView<const FVector3f> BoneVelocities,
 		const TArrayView<const int32> BoneIndices,
-		const FLinearColor& Color,
-		const float Scale = 100.0f,
+		const FLinearColor& Color = FLinearColor(0.0f, 0.0f, 0.025f, 1.0f),
 		const float Thickness = 0.0f,
-		const bool bDepthTest = false);
+		const bool bDepthTest = false,
+		const float Scale = 100.0f);
 
-	/** Draws capsules for a skeleton */
-	static UE_API void DrawDebugSkeletonCapsulesArrayView(
-		const FDebugDrawer& Drawer,
-		const TArrayView<const FVector> BoneLocations,
-		const TArrayView<const FQuat4f> BoneRotations,
-		const TArrayView<const FVector3f> BoneLocalLocations,
-		const TArrayView<const int32> BoneIndices,
-		const TArrayView<const int32> BoneParents,
-		const float MinCapsuleRadius = 1.0f,
-		const float MaxCapsuleRadius = 10.0f,
-		const FLinearColor Color = FLinearColor::Red,
-		const float Thickness = 0.0f,
-		const int32 Segments = 9,
-		const bool bDepthTest = true);
+	/** Debug Draw a pose made up of bone locations and velocities */
+	UFUNCTION(BlueprintCallable, Category = "Draw Debug Library", meta = (DevelopmentOnly, AutoCreateRefTerm = "RelativeTransform"))
+	static UE_API void DrawDebugPose(const FDebugDrawer& Drawer, const TArray<FVector>& BoneLocations, const TArray<FVector>& BoneLinearVelocities, const FTransform& RelativeTransform = FTransform(), const FLinearColor Color = FLinearColor::Red, const float Thickness = 0.0f, const bool bDepthTest = true, const float DrawVelocityLineScale = 1.0f);
+	static UE_API void DrawDebugPoseFromArrayViews(const FDebugDrawer& Drawer, const TArrayView<const FVector> BoneLocations, const TArrayView<const FVector> BoneLinearVelocities, const FTransform& RelativeTransform = FTransform(), const FLinearColor Color = FLinearColor::Red, const float Thickness = 0.0f, const bool bDepthTest = true, const float DrawVelocityLineScale = 1.0f);
 
 public: // Text
 
 	/** Get the number of line segments required to draw debug string. */
-	static UE_API int32 DrawDebugStringViewSegmentNum(const FStringView String, const bool bMonospaced = true);
+	UFUNCTION(BlueprintCallable, Category = "Draw Debug Library", meta = (DevelopmentOnly, AutoCreateRefTerm = "String, Settings"))
+	static UE_API int32 DrawDebugStringSegmentNum(const FString& String, const FDrawDebugStringSettings& Settings = FDrawDebugStringSettings());
+	static UE_API int32 DrawDebugStringViewSegmentNum(const FStringView String, const FDrawDebugStringSettings& Settings = FDrawDebugStringSettings());
 
 	/** Get the dimensions of a draw debug string. Useful for aligning or centering text. */
-	UFUNCTION(BlueprintCallable, Category = "Draw Debug Library", meta = (DevelopmentOnly, AutoCreateRefTerm = "String"))
-	static UE_API FVector DrawDebugStringDimensions(const FString& String, const float Scale = 10.0f, const bool bMonospaced = true);
-	static UE_API FVector DrawDebugStringViewDimensions(const FStringView String, const float Scale = 10.0f, const bool bMonospaced = true);
+	UFUNCTION(BlueprintCallable, Category = "Draw Debug Library", meta = (DevelopmentOnly, AutoCreateRefTerm = "String, Settings"))
+	static UE_API FVector DrawDebugStringDimensions(const FString& String, const FDrawDebugStringSettings& Settings = FDrawDebugStringSettings());
+	static UE_API FVector DrawDebugStringViewDimensions(const FStringView String, const FDrawDebugStringSettings& Settings = FDrawDebugStringSettings());
 
 	/** Get the dimensions of a draw debug text. Useful for aligning or centering text. */
-	UFUNCTION(BlueprintCallable, Category = "Draw Debug Library", meta = (DevelopmentOnly, AutoCreateRefTerm = "Text"))
-	static UE_API FVector DrawDebugTextDimensions(const FText& Text, const float Scale = 10.0f, const bool bMonospaced = true);
+	UFUNCTION(BlueprintCallable, Category = "Draw Debug Library", meta = (DevelopmentOnly, AutoCreateRefTerm = "Text, Settings"))
+	static UE_API FVector DrawDebugTextDimensions(const FText& Text, const FDrawDebugStringSettings& Settings = FDrawDebugStringSettings());
 
 	/** Get the dimensions of a draw debug name. Useful for aligning or centering text. */
-	UFUNCTION(BlueprintCallable, Category = "Draw Debug Library", meta = (DevelopmentOnly, AutoCreateRefTerm = "Name"))
-	static UE_API FVector DrawDebugNameDimensions(const FName& Name, const float Scale = 10.0f, const bool bMonospaced = true);
+	UFUNCTION(BlueprintCallable, Category = "Draw Debug Library", meta = (DevelopmentOnly, AutoCreateRefTerm = "Name, Settings"))
+	static UE_API FVector DrawDebugNameDimensions(const FName& Name, const FDrawDebugStringSettings& Settings = FDrawDebugStringSettings());
 
 	/** Debug Draw a string. Will only render ASCII characters. */
-	UFUNCTION(BlueprintCallable, Category = "Draw Debug Library", meta = (DevelopmentOnly, AutoCreateRefTerm = "String, Location, Rotation"))
-	static UE_API void DrawDebugString(const FDebugDrawer& Drawer, const FString& String, const FVector& Location, const FRotator& Rotation, const float Scale = 10.0f, const bool bMonospaced = true, const FLinearColor Color = FLinearColor::Red, const float Thickness = 0.0f, const bool bDepthTest = true);
-	static UE_API void DrawDebugStringView(const FDebugDrawer& Drawer, const FStringView String, const FVector& Location, const FRotator& Rotation, const float Scale = 10.0f, const bool bMonospaced = true, const FLinearColor Color = FLinearColor::Red, const float Thickness = 0.0f, const bool bDepthTest = true);
+	UFUNCTION(BlueprintCallable, Category = "Draw Debug Library", meta = (DevelopmentOnly, AutoCreateRefTerm = "String, Location, Rotation, Settings"))
+	static UE_API void DrawDebugString(const FDebugDrawer& Drawer, const FString& String, const FVector& Location, const FRotator& Rotation, const FLinearColor Color = FLinearColor::Red, const float Thickness = 0.0f, const bool bDepthTest = true, const FDrawDebugStringSettings& Settings = FDrawDebugStringSettings());
+	static UE_API void DrawDebugStringView(const FDebugDrawer& Drawer, const FStringView String, const FVector& Location, const FRotator& Rotation, const FLinearColor Color = FLinearColor::Red, const float Thickness = 0.0f, const bool bDepthTest = true, const FDrawDebugStringSettings& Settings = FDrawDebugStringSettings());
 
 	/** Debug Draw a name */
-	UFUNCTION(BlueprintCallable, Category = "Draw Debug Library", meta = (DevelopmentOnly, AutoCreateRefTerm = "Name, Location, Rotation"))
-	static UE_API void DrawDebugName(const FDebugDrawer& Drawer, const FName& Name, const FVector& Location, const FRotator& Rotation, const float Scale = 10.0f, const bool bMonospaced = true, const FLinearColor Color = FLinearColor::Red, const float Thickness = 0.0f, const bool bDepthTest = true);
+	UFUNCTION(BlueprintCallable, Category = "Draw Debug Library", meta = (DevelopmentOnly, AutoCreateRefTerm = "Name, Location, Rotation, Settings"))
+	static UE_API void DrawDebugName(const FDebugDrawer& Drawer, const FName& Name, const FVector& Location, const FRotator& Rotation, const FLinearColor Color = FLinearColor::Red, const float Thickness = 0.0f, const bool bDepthTest = true, const FDrawDebugStringSettings& Settings = FDrawDebugStringSettings());
 
 	/** Debug Draw some text. Will only render ASCII characters. */
-	UFUNCTION(BlueprintCallable, Category = "Draw Debug Library", meta = (DevelopmentOnly, AutoCreateRefTerm = "Text, Location, Rotation"))
-	static UE_API void DrawDebugText(const FDebugDrawer& Drawer, const FText& Text, const FVector& Location, const FRotator& Rotation, const float Scale = 10.0f, const bool bMonospaced = true, const FLinearColor Color = FLinearColor::Red, const float Thickness = 0.0f, const bool bDepthTest = true);
+	UFUNCTION(BlueprintCallable, Category = "Draw Debug Library", meta = (DevelopmentOnly, AutoCreateRefTerm = "Text, Location, Rotation, Settings"))
+	static UE_API void DrawDebugText(const FDebugDrawer& Drawer, const FText& Text, const FVector& Location, const FRotator& Rotation, const FLinearColor Color = FLinearColor::Red, const float Thickness = 0.0f, const bool bDepthTest = true, const FDrawDebugStringSettings& Settings = FDrawDebugStringSettings());
 
 public: // Visual Logger Text
 
 	/** Debug Draw a string to the visual logger. Will do nothing if not using a VisualLoggerDrawer. Will not display on-screen during recording. */
 	UFUNCTION(BlueprintCallable, Category = "Draw Debug Library", meta = (DevelopmentOnly, AutoCreateRefTerm = "String, Location"))
-	static UE_API void DrawDebugVisualLoggerString(const FDebugDrawer& Drawer, const FString& String, const FVector& Location, const FLinearColor Color = FLinearColor::Red);
-	static UE_API void DrawDebugVisualLoggerStringView(const FDebugDrawer& Drawer, const FStringView String, const FVector& Location, const FLinearColor Color = FLinearColor::Red);
+	static UE_API void VisualLoggerDrawString(const FDebugDrawer& Drawer, const FString& String, const FVector& Location, const FLinearColor Color = FLinearColor::Red);
+	static UE_API void VisualLoggerDrawStringView(const FDebugDrawer& Drawer, const FStringView String, const FVector& Location, const FLinearColor Color = FLinearColor::Red);
 
 	/** Debug Draw a name to the visual logger. Will do nothing if not using a VisualLoggerDrawer. Will not display on-screen during recording. */
 	UFUNCTION(BlueprintCallable, Category = "Draw Debug Library", meta = (DevelopmentOnly, AutoCreateRefTerm = "Name, Location"))
-	static UE_API void DrawDebugVisualLoggerName(const FDebugDrawer& Drawer, const FName& Name, const FVector& Location, const FLinearColor Color = FLinearColor::Red);
+	static UE_API void VisualLoggerDrawName(const FDebugDrawer& Drawer, const FName& Name, const FVector& Location, const FLinearColor Color = FLinearColor::Red);
 
 	/** Debug Draw some text to the visual logger. Will do nothing if not using a VisualLoggerDrawer. Will not display on-screen during recording. */
 	UFUNCTION(BlueprintCallable, Category = "Draw Debug Library", meta = (DevelopmentOnly, AutoCreateRefTerm = "Text, Location"))
-	static UE_API void DrawDebugVisualLoggerText(const FDebugDrawer& Drawer, const FText& Text, const FVector& Location, const FLinearColor Color = FLinearColor::Red);
+	static UE_API void VisualLoggerDrawText(const FDebugDrawer& Drawer, const FText& Text, const FVector& Location, const FLinearColor Color = FLinearColor::Red);
 
 
 public: // Graphs
 
 	/** Debug Draw a simple graph. */
-	UFUNCTION(BlueprintCallable, Category = "Draw Debug Library", meta = (DevelopmentOnly, AutoCreateRefTerm = "Location, Rotation, Title, XaxisLabel, YaxisLabel"))
+	UFUNCTION(BlueprintCallable, Category = "Draw Debug Library", meta = (DevelopmentOnly, AutoCreateRefTerm = "Location, Rotation, AxesSettings"))
 	static UE_API void DrawDebugGraph(
 		const FDebugDrawer& Drawer,
 		const FVector& Location,
@@ -509,17 +831,14 @@ public: // Graphs
 		const float Ymax = 1.0f,
 		const float XaxisLength = 100.0f,
 		const float YaxisLength = 100.0f,
-		const FString& Title = TEXT(""),
-		const FString& XaxisLabel = TEXT(""),
-		const FString& YaxisLabel = TEXT(""),
-		const float TextScale = 10.0f,
 		const FLinearColor Color = FLinearColor::Red,
 		const float AxisThickness = 1.0f,
 		const float LineThickness = 0.0f,
 		const float TextThickness = 0.0f,
-		const bool bDepthTest = true);
+		const bool bDepthTest = true,
+		const FDrawDebugGraphAxesSettings& AxesSettings = FDrawDebugGraphAxesSettings());
 
-	static UE_API void DrawDebugGraphView(
+	static UE_API void DrawDebugGraphArrayView(
 		const FDebugDrawer& Drawer,
 		const FVector& Location,
 		const FRotator& Rotation,
@@ -531,15 +850,89 @@ public: // Graphs
 		const float Ymax = 1.0f,
 		const float XaxisLength = 100.0f,
 		const float YaxisLength = 100.0f,
-		const FStringView Title = TEXT(""),
-		const FStringView XaxisLabel = TEXT(""),
-		const FStringView YaxisLabel = TEXT(""),
-		const float TextScale = 10.0f,
 		const FLinearColor Color = FLinearColor::Red,
 		const float AxisThickness = 1.0f,
 		const float LineThickness = 0.0f,
 		const float TextThickness = 0.0f,
+		const bool bDepthTest = true,
+		const FDrawDebugGraphAxesSettings& AxesSettings = FDrawDebugGraphAxesSettings());
+
+	/** Debug Draw a simple graph axes. */
+	UFUNCTION(BlueprintCallable, Category = "Draw Debug Library", meta = (DevelopmentOnly, AutoCreateRefTerm = "Location, Rotation, AxesSettings"))
+	static UE_API void DrawDebugGraphAxes(
+		const FDebugDrawer& Drawer,
+		const FVector& Location,
+		const FRotator& Rotation,
+		const float XaxisLength = 100.0f,
+		const float YaxisLength = 100.0f,
+		const FLinearColor Color = FLinearColor::Red,
+		const float AxisThickness = 1.0f,
+		const float TextThickness = 0.0f,
+		const bool bDepthTest = true,
+		const FDrawDebugGraphAxesSettings& AxesSettings = FDrawDebugGraphAxesSettings());
+
+	/** Debug Draw a line on a simple graph. */
+	UFUNCTION(BlueprintCallable, Category = "Draw Debug Library", meta = (DevelopmentOnly, AutoCreateRefTerm = "Location, Rotation, AxesSettings"))
+	static UE_API void DrawDebugGraphLine(
+		const FDebugDrawer& Drawer,
+		const FVector& Location,
+		const FRotator& Rotation,
+		const TArray<float>& Xvalues,
+		const TArray<float>& Yvalues,
+		const float Xmin = 0.0f,
+		const float Xmax = 1.0f,
+		const float Ymin = 0.0f,
+		const float Ymax = 1.0f,
+		const float XaxisLength = 100.0f,
+		const float YaxisLength = 100.0f,
+		const FLinearColor Color = FLinearColor::Red,
+		const float LineThickness = 0.0f,
 		const bool bDepthTest = true);
+
+	static UE_API void DrawDebugGraphLineArrayView(
+		const FDebugDrawer& Drawer,
+		const FVector& Location,
+		const FRotator& Rotation,
+		const TArrayView<const float> Xvalues,
+		const TArrayView<const float> Yvalues,
+		const float Xmin = 0.0f,
+		const float Xmax = 1.0f,
+		const float Ymin = 0.0f,
+		const float Ymax = 1.0f,
+		const float XaxisLength = 100.0f,
+		const float YaxisLength = 100.0f,
+		const FLinearColor Color = FLinearColor::Red,
+		const float LineThickness = 0.0f,
+		const bool bDepthTest = true);
+
+	/** Debug Draw a simple graph legend. */
+	UFUNCTION(BlueprintCallable, Category = "Draw Debug Library", meta = (DevelopmentOnly, AutoCreateRefTerm = "Location, Rotation, LegendSettings"))
+	static UE_API void DrawDebugGraphLegend(
+		const FDebugDrawer& Drawer,
+		const FVector& Location,
+		const FRotator& Rotation,
+		const TArray<FLinearColor>& LegendColors,
+		const TArray<FString>& LegendLabels,
+		const float XaxisLength = 100.0f,
+		const float YaxisLength = 100.0f,
+		const float TextThickness = 0.0f,
+		const float PointSize = 2.5f,
+		const bool bDepthTest = true,
+		const FDrawDebugStringSettings& LegendSettings = FDrawDebugStringSettings());
+
+	static UE_API void DrawDebugGraphLegendArrayView(
+		const FDebugDrawer& Drawer,
+		const FVector& Location,
+		const FRotator& Rotation,
+		const TArrayView<const FLinearColor> LegendColors,
+		const TArrayView<const FString> LegendLabels,
+		const float XaxisLength = 100.0f,
+		const float YaxisLength = 100.0f,
+		const float TextThickness = 0.0f,
+		const float PointSize = 2.5f,
+		const bool bDepthTest = true,
+		const FDrawDebugStringSettings& LegendSettings = FDrawDebugStringSettings());
+
 };
 
 #undef UE_API
