Fixed normalization issue in DrawDebugLibrary

diff --git a/Engine/Plugins/Experimental/DrawDebugLibrary/Source/DrawDebugLibrary/Private/DrawDebugLibrary.cpp b/Engine/Plugins/Experimental/DrawDebugLibrary/Source/DrawDebugLibrary/Private/DrawDebugLibrary.cpp
--- a/Engine/Plugins/Experimental/DrawDebugLibrary/Source/DrawDebugLibrary/Private/DrawDebugLibrary.cpp
+++ b/Engine/Plugins/Experimental/DrawDebugLibrary/Source/DrawDebugLibrary/Private/DrawDebugLibrary.cpp
@@ -2580,12 +2580,15 @@ void UDrawDebugLibrary::DrawDebugGraphLineArrayView(
 	TArray<TPair<FVector, FVector>> LineSegments;
 	LineSegments.Reserve(PointNum - 1);
 
+	const float XScale = FMath::Max(Xmax - Xmin, UE_SMALL_NUMBER);
+	const float YScale = FMath::Max(Ymax - Ymin, UE_SMALL_NUMBER);
+
 	for (int32 PointIdx = 0; PointIdx < PointNum - 1; PointIdx++)
 	{
-		const float Xloc0 = XaxisLength * ((Xvalues[PointIdx + 0] - Xmin) / (Xmax - Xmin));
-		const float Yloc0 = YaxisLength * ((Yvalues[PointIdx + 0] - Ymin) / (Ymax - Ymin));
-		const float Xloc1 = XaxisLength * ((Xvalues[PointIdx + 1] - Xmin) / (Xmax - Xmin));
-		const float Yloc1 = YaxisLength * ((Yvalues[PointIdx + 1] - Ymin) / (Ymax - Ymin));
+		const float Xloc0 = XaxisLength * ((Xvalues[PointIdx + 0] - Xmin) / XScale);
+		const float Yloc0 = YaxisLength * ((Yvalues[PointIdx + 0] - Ymin) / YScale);
+		const float Xloc1 = XaxisLength * ((Xvalues[PointIdx + 1] - Xmin) / XScale);
+		const float Yloc1 = YaxisLength * ((Yvalues[PointIdx + 1] - Ymin) / YScale);
 
 		LineSegments.Add({
 			Location + Rotation.RotateVector(FVector(Xloc0, 0.0, Yloc0)),
