DrawDebugLibrary: Added option for not drawing to the scene during recording

#virtualized

diff --git a/Engine/Plugins/Experimental/DrawDebugLibrary/Source/DrawDebugLibrary/Private/DrawDebugLibrary.cpp b/Engine/Plugins/Experimental/DrawDebugLibrary/Source/DrawDebugLibrary/Private/DrawDebugLibrary.cpp
--- a/Engine/Plugins/Experimental/DrawDebugLibrary/Source/DrawDebugLibrary/Private/DrawDebugLibrary.cpp
+++ b/Engine/Plugins/Experimental/DrawDebugLibrary/Source/DrawDebugLibrary/Private/DrawDebugLibrary.cpp
@@ -453,23 +453,35 @@ FDebugDrawer FDebugDrawer::MakeDebugDrawer(FPrimitiveDrawInterface* PrimitiveDra
 	return FDebugDrawer(PrimitiveDrawInterface);
 }
 
-FDebugDrawer FDebugDrawer::MakeVisualLoggerDebugDrawer(UObject* Object, const FName Category, const EDrawDebugLogVerbosity Verbosity, const bool bDrawToScene)
+FDebugDrawer FDebugDrawer::MakeVisualLoggerDebugDrawer(
+	UObject* Object, 
+	const FName Category, 
+	const EDrawDebugLogVerbosity Verbosity, 
+	const bool bDrawToScene,
+	const bool bDrawToSceneWhileRecording)
 {
 	FDebugDrawer Drawer;
-	Drawer.World = (Object && bDrawToScene) ? Object->GetWorld() : nullptr;
 	Drawer.VisualLoggerObject = Object;
 	Drawer.VisualLoggerCategory = Category;
 	Drawer.VisualLoggerVerbosity = UE::DrawDebugLibrary::Private::LogVerbosityFromDrawDebugLogVerbosity(Verbosity);
+	Drawer.bVisualLoggerDrawToScene = bDrawToScene;
+	Drawer.bVisualLoggerDrawToSceneWhileRecording = bDrawToSceneWhileRecording;
 	return Drawer;
 }
 
-FDebugDrawer FDebugDrawer::MakeVisualLoggerDebugDrawer(UObject* Object, const FLogCategoryBase& Category, const ELogVerbosity::Type Verbosity, const bool bDrawToScene)
+FDebugDrawer FDebugDrawer::MakeVisualLoggerDebugDrawer(
+	UObject* Object, 
+	const FLogCategoryBase& Category, 
+	const ELogVerbosity::Type Verbosity, 
+	const bool bDrawToScene,
+	const bool bDrawToSceneWhileRecording)
 {
 	FDebugDrawer Drawer;
-	Drawer.World = (Object && bDrawToScene) ? Object->GetWorld() : nullptr;
 	Drawer.VisualLoggerObject = Object;
 	Drawer.VisualLoggerCategory = Category.GetCategoryName();
 	Drawer.VisualLoggerVerbosity = Verbosity;
+	Drawer.bVisualLoggerDrawToScene = bDrawToScene;
+	Drawer.bVisualLoggerDrawToSceneWhileRecording = bDrawToSceneWhileRecording;
 	return Drawer;
 }
 
@@ -487,6 +499,8 @@ FDebugDrawer FDebugDrawer::MakeMergedDebugDrawer(const TArrayView<const FDebugDr
 			Drawer.VisualLoggerObject = SubDrawer.VisualLoggerObject;
 			Drawer.VisualLoggerCategory = SubDrawer.VisualLoggerCategory;
 			Drawer.VisualLoggerVerbosity = SubDrawer.VisualLoggerVerbosity;
+			Drawer.bVisualLoggerDrawToScene = SubDrawer.bVisualLoggerDrawToScene;
+			Drawer.bVisualLoggerDrawToSceneWhileRecording = SubDrawer.bVisualLoggerDrawToSceneWhileRecording;
 		}
 	}
 
@@ -529,7 +543,7 @@ void FDebugDrawer::DrawPoints(const TArrayView<const FVector> Locations, const F
 		FVisualLogEntry* CurrentEntry = nullptr;
 		if (FVisualLogger::CheckVisualLogInputInternal(VisualLoggerObject, VisualLoggerCategory, VisualLoggerVerbosity, /*OutWorld=*/ nullptr, &CurrentEntry))
 		{
-			// The uin16 integer size is difficult to match to the float thickness used by the other debug drawers. This is my best estimation but I
+			// The uint16 integer size is difficult to match to the float thickness used by the other debug drawers. This is my best estimation but I
 			// am not 100% sure the mapping is correct.
 			const uint16 IntegerSize = Thickness == 0.0f ? 0 : FMath::CeilToInt(FMath::Max(2.0f * Thickness, 0.0f));
 			FVisualLogShapeElement Element(TEXT(""), Color.ToFColor(false), IntegerSize, VisualLoggerCategory);
@@ -539,6 +553,16 @@ void FDebugDrawer::DrawPoints(const TArrayView<const FVector> Locations, const F
 			Element.Points.Append(Locations);
 			CurrentEntry->ElementsToDraw.Add(Element);
 		}
+		
+		const UWorld* VisualLoggerWorld = VisualLoggerObject->GetWorld();
+
+		if (VisualLoggerWorld && bVisualLoggerDrawToScene && (bVisualLoggerDrawToSceneWhileRecording || !FVisualLogger::IsRecording()))
+		{
+			for (int32 Idx = 0; Idx < Num; Idx++)
+			{
+				::DrawDebugPoint(VisualLoggerWorld, Locations[Idx], Thickness, Color.ToFColor(true), false, -1.0f, bDepthTest ? 0 : 1);
+			}
+		}
 	}
 #endif
 }
@@ -598,6 +622,16 @@ void FDebugDrawer::DrawLines(const TArrayView<const FVector> StartLocations, con
 			}
 			CurrentEntry->ElementsToDraw.Add(Element);
 		}
+
+		const UWorld* VisualLoggerWorld = VisualLoggerObject->GetWorld();
+
+		if (VisualLoggerWorld && bVisualLoggerDrawToScene && (bVisualLoggerDrawToSceneWhileRecording || !FVisualLogger::IsRecording()))
+		{
+			for (int32 Idx = 0; Idx < Num; Idx++)
+			{
+				::DrawDebugLine(VisualLoggerWorld, StartLocations[Idx], EndLocations[Idx], Color.ToFColor(true), false, -1.0f, bDepthTest ? 0 : 1, Thickness);
+			}
+		}
 	}
 #endif
 }
@@ -657,6 +691,16 @@ void FDebugDrawer::DrawLinesPairs(const TArrayView<const TPair<FVector, FVector>
 			}
 			CurrentEntry->ElementsToDraw.Add(Element);
 		}
+
+		const UWorld* VisualLoggerWorld = VisualLoggerObject->GetWorld();
+
+		if (VisualLoggerWorld && bVisualLoggerDrawToScene && (bVisualLoggerDrawToSceneWhileRecording || !FVisualLogger::IsRecording()))
+		{
+			for (int32 Idx = 0; Idx < Num; Idx++)
+			{
+				::DrawDebugLine(VisualLoggerWorld, Segments[Idx].Key, Segments[Idx].Value, Color.ToFColor(true), false, -1.0f, bDepthTest ? 0 : 1, Thickness);
+			}
+		}
 	}
 #endif
 }
@@ -723,19 +767,19 @@ FDebugDrawer UDrawDebugLibrary::MakePDIDebugDrawer(FPrimitiveDrawInterface* Prim
 	return FDebugDrawer::MakeDebugDrawer(PrimitiveDrawInterface);
 }
 
-FDebugDrawer UDrawDebugLibrary::MakeVisualLoggerDebugDrawer(UObject* Object, const FName Category, const EDrawDebugLogVerbosity Verbosity, const bool bDrawToScene)
+FDebugDrawer UDrawDebugLibrary::MakeVisualLoggerDebugDrawer(UObject* Object, const FName Category, const EDrawDebugLogVerbosity Verbosity, const bool bDrawToScene,const bool bDrawToSceneWhileRecording)
 {
-	return FDebugDrawer::MakeVisualLoggerDebugDrawer(Object, Category, Verbosity, bDrawToScene);
+	return FDebugDrawer::MakeVisualLoggerDebugDrawer(Object, Category, Verbosity, bDrawToScene, bDrawToSceneWhileRecording);
 }
 
-FDebugDrawer UDrawDebugLibrary::MakeVisualLoggerDebugDrawerFromObject(UObject* Object, const FName Category, const EDrawDebugLogVerbosity Verbosity, const bool bDrawToScene)
+FDebugDrawer UDrawDebugLibrary::MakeVisualLoggerDebugDrawerFromObject(UObject* Object, const FName Category, const EDrawDebugLogVerbosity Verbosity, const bool bDrawToScene, const bool bDrawToSceneWhileRecording)
 {
-	return FDebugDrawer::MakeVisualLoggerDebugDrawer(Object, Category, Verbosity, bDrawToScene);
+	return FDebugDrawer::MakeVisualLoggerDebugDrawer(Object, Category, Verbosity, bDrawToScene, bDrawToSceneWhileRecording);
 }
 
-FDebugDrawer UDrawDebugLibrary::MakeVisualLoggerDebugDrawerFromObjectWithCategory(UObject* Object, const FLogCategoryBase& Category, const ELogVerbosity::Type Verbosity, const bool bDrawToScene)
+FDebugDrawer UDrawDebugLibrary::MakeVisualLoggerDebugDrawerFromObjectWithCategory(UObject* Object, const FLogCategoryBase& Category, const ELogVerbosity::Type Verbosity, const bool bDrawToScene, const bool bDrawToSceneWhileRecording)
 {
-	return FDebugDrawer::MakeVisualLoggerDebugDrawer(Object, Category, Verbosity, bDrawToScene);
+	return FDebugDrawer::MakeVisualLoggerDebugDrawer(Object, Category, Verbosity, bDrawToScene, bDrawToSceneWhileRecording);
 }
 
 FDebugDrawer UDrawDebugLibrary::MakeMergedDebugDrawer(const TArray<FDebugDrawer>& DebugDrawers)
diff --git a/Engine/Plugins/Experimental/DrawDebugLibrary/Source/DrawDebugLibrary/Public/DrawDebugLibrary.h b/Engine/Plugins/Experimental/DrawDebugLibrary/Source/DrawDebugLibrary/Public/DrawDebugLibrary.h
--- a/Engine/Plugins/Experimental/DrawDebugLibrary/Source/DrawDebugLibrary/Public/DrawDebugLibrary.h
+++ b/Engine/Plugins/Experimental/DrawDebugLibrary/Source/DrawDebugLibrary/Public/DrawDebugLibrary.h
@@ -324,10 +324,20 @@ public: // Constructors
 	static UE_API FDebugDrawer MakeDebugDrawer(FPrimitiveDrawInterface* PrimitiveDrawInterface);
 
 	/** Make a debug drawer for the Visual Logger */
-	static UE_API FDebugDrawer MakeVisualLoggerDebugDrawer(UObject* Object, const FName Category = TEXT("LogDrawDebugLibrary"), const EDrawDebugLogVerbosity Verbosity = EDrawDebugLogVerbosity::Display, const bool bDrawToScene = true);
+	static UE_API FDebugDrawer MakeVisualLoggerDebugDrawer(
+		UObject* Object, 
+		const FName Category = TEXT("LogDrawDebugLibrary"), 
+		const EDrawDebugLogVerbosity Verbosity = EDrawDebugLogVerbosity::Display, 
+		const bool bDrawToScene = true,
+		const bool bDrawToSceneWhileRecording = true);
 
 	/** Make a debug drawer for the Visual Logger */
-	static UE_API FDebugDrawer MakeVisualLoggerDebugDrawer(UObject* Object, const FLogCategoryBase& Category, const ELogVerbosity::Type Verbosity = ELogVerbosity::Display, const bool bDrawToScene = true);
+	static UE_API FDebugDrawer MakeVisualLoggerDebugDrawer(
+		UObject* Object, 
+		const FLogCategoryBase& Category, 
+		const ELogVerbosity::Type Verbosity = ELogVerbosity::Display, 
+		const bool bDrawToScene = true,
+		const bool bDrawToSceneWhileRecording = true);
 
 	/** Make a debug drawer by merging multiple other debug drawers of different types */
 	static UE_API FDebugDrawer MakeMergedDebugDrawer(const TArrayView<const FDebugDrawer> DebugDrawers);
@@ -362,6 +372,12 @@ private:
 
 	/** Visual Logger Verbosity. */
 	ELogVerbosity::Type VisualLoggerVerbosity = ELogVerbosity::Display;
+
+	/** Visual Logger if to also draw to scene */
+	bool bVisualLoggerDrawToScene = true;
+
+	/** Visual Logger if to also draw while recording */
+	bool bVisualLoggerDrawToSceneWhileRecording = true;
 };
 
 
@@ -403,12 +419,28 @@ public: // Drawers
 
 	/** Make a debug drawer for the Visual Logger from the current "Self" object */
 	UFUNCTION(BlueprintPure, Category = "Draw Debug Library", meta = (DevelopmentOnly, DefaultToSelf = "Object", HidePin = "Object"))
-	static UE_API FDebugDrawer MakeVisualLoggerDebugDrawer(UObject* Object, const FName Category = TEXT("LogDrawDebugLibrary"), const EDrawDebugLogVerbosity Verbosity = EDrawDebugLogVerbosity::Display, const bool bDrawToScene = true);
+	static UE_API FDebugDrawer MakeVisualLoggerDebugDrawer(
+		UObject* Object, 
+		const FName Category = TEXT("LogDrawDebugLibrary"), 
+		const EDrawDebugLogVerbosity Verbosity = EDrawDebugLogVerbosity::Display, 
+		const bool bDrawToScene = true,
+		const bool bDrawToSceneWhileRecording = true);
 
 	/** Make a debug drawer for the Visual Logger */
 	UFUNCTION(BlueprintPure, Category = "Draw Debug Library", meta = (DevelopmentOnly))
-	static UE_API FDebugDrawer MakeVisualLoggerDebugDrawerFromObject(UObject* Object, const FName Category = TEXT("LogDrawDebugLibrary"), const EDrawDebugLogVerbosity Verbosity = EDrawDebugLogVerbosity::Display, const bool bDrawToScene = true);
-	static UE_API FDebugDrawer MakeVisualLoggerDebugDrawerFromObjectWithCategory(UObject* Object, const FLogCategoryBase& Category, const ELogVerbosity::Type Verbosity = ELogVerbosity::Display, const bool bDrawToScene = true);
+	static UE_API FDebugDrawer MakeVisualLoggerDebugDrawerFromObject(
+		UObject* Object, 
+		const FName Category = TEXT("LogDrawDebugLibrary"), 
+		const EDrawDebugLogVerbosity Verbosity = EDrawDebugLogVerbosity::Display, 
+		const bool bDrawToScene = true, 
+		const bool bDrawToSceneWhileRecording = true);
+
+	static UE_API FDebugDrawer MakeVisualLoggerDebugDrawerFromObjectWithCategory(
+		UObject* Object, 
+		const FLogCategoryBase& Category, 
+		const ELogVerbosity::Type Verbosity = ELogVerbosity::Display, 
+		const bool bDrawToScene = true, 
+		const bool bDrawToSceneWhileRecording = true);
 
 	/** Make a debug drawer by merging multiple other debug drawers of different types */
 	UFUNCTION(BlueprintPure, Category = "Draw Debug Library", meta = (DevelopmentOnly))
