DrawDebugLibrary: Added ability to change line style

diff --git a/Engine/Plugins/Experimental/DrawDebugLibrary/Source/DrawDebugLibrary/Private/DrawDebugLibrary.cpp b/Engine/Plugins/Experimental/DrawDebugLibrary/Source/DrawDebugLibrary/Private/DrawDebugLibrary.cpp
--- a/Engine/Plugins/Experimental/DrawDebugLibrary/Source/DrawDebugLibrary/Private/DrawDebugLibrary.cpp
+++ b/Engine/Plugins/Experimental/DrawDebugLibrary/Source/DrawDebugLibrary/Private/DrawDebugLibrary.cpp
@@ -522,84 +522,27 @@ void FDebugDrawer::DrawPoints(const TArrayView<const FVector> Locations, const F
 #if UE_ENABLE_DEBUG_DRAWING
 	if (World)
 	{
-		for (int32 Idx = 0; Idx < Num; Idx++)
-		{
-			::DrawDebugPoint(World, Locations[Idx], Thickness, Color.ToFColor(true), false, -1.0f, bDepthTest ? 0 : 1);
-		}
-	}
-
-	if (AnimInstanceProxy)
-	{
-		for (int32 Idx = 0; Idx < Num; Idx++)
-		{
-			AnimInstanceProxy->AnimDrawDebugPoint(Locations[Idx], Thickness, Color.ToFColor(true), false, -1.0f, bDepthTest ? SDPG_World : SDPG_Foreground);
-		}
-	}
-#endif
-
-#if ENABLE_VISUAL_LOG
-	if (VisualLoggerObject)
-	{
-		FVisualLogEntry* CurrentEntry = nullptr;
-		if (FVisualLogger::CheckVisualLogInputInternal(VisualLoggerObject, VisualLoggerCategory, VisualLoggerVerbosity, /*OutWorld=*/ nullptr, &CurrentEntry))
-		{
-			// The uint16 integer size is difficult to match to the float thickness used by the other debug drawers. This is my best estimation but I
-			// am not 100% sure the mapping is correct.
-			const uint16 IntegerSize = Thickness == 0.0f ? 0 : FMath::CeilToInt(FMath::Max(2.0f * Thickness, 0.0f));
-			FVisualLogShapeElement Element(TEXT(""), Color.ToFColor(false), IntegerSize, VisualLoggerCategory);
-			Element.Type = EVisualLoggerShapeElement::SinglePoint;
-			Element.Radius = 0.0f;
-			Element.Verbosity = VisualLoggerVerbosity;
-			Element.Points.Append(Locations);
-			CurrentEntry->ElementsToDraw.Add(Element);
-		}
-		
-		const UWorld* VisualLoggerWorld = VisualLoggerObject->GetWorld();
-
-		if (VisualLoggerWorld && bVisualLoggerDrawToScene && (bVisualLoggerDrawToSceneWhileRecording || !FVisualLogger::IsRecording()))
+		if (Thickness == 0.0f)
 		{
 			for (int32 Idx = 0; Idx < Num; Idx++)
 			{
-				::DrawDebugPoint(VisualLoggerWorld, Locations[Idx], Thickness, Color.ToFColor(true), false, -1.0f, bDepthTest ? 0 : 1);
+				::DrawDebugPoint(World, Locations[Idx], 1.0f, Color.ToFColor(true), false, -1.0f, bDepthTest ? 0 : 1);
 			}
 		}
-	}
-#endif
-}
-
-void FDebugDrawer::DrawLines(const TArrayView<const FVector> StartLocations, const TArrayView<const FVector> EndLocations, const FLinearColor& Color, const float Thickness, const bool bDepthTest) const
-{
-	const int32 Num = FMath::Min(StartLocations.Num(), EndLocations.Num());
-
-	if (PDI)
-	{
-		for (int32 Idx = 0; Idx < Num; Idx++)
+		else
 		{
-			if (bDepthTest)
-			{
-				PDI->DrawTranslucentLine(StartLocations[Idx], EndLocations[Idx], Color, SDPG_World, FMath::Max(Thickness, 0.0f), 0.0f, Thickness <= 0.0f);
-			}
-			else
+			for (int32 Idx = 0; Idx < Num; Idx++)
 			{
-				PDI->DrawTranslucentLine(StartLocations[Idx], EndLocations[Idx], Color, Thickness == 0.0f ? SDPG_Foreground : SDPG_World, Thickness, Thickness == 0.0f ? 1.0f : UE_MAX_FLT);
+				::DrawDebugLine(World, Locations[Idx], Locations[Idx], Color.ToFColor(true), false, -1.0f, bDepthTest ? 0 : 1, Thickness);
 			}
 		}
 	}
 
-#if UE_ENABLE_DEBUG_DRAWING
-	if (World)
-	{
-		for (int32 Idx = 0; Idx < Num; Idx++)
-		{
-			::DrawDebugLine(World, StartLocations[Idx], EndLocations[Idx], Color.ToFColor(true), false, -1.0f, bDepthTest ? 0 : 1, Thickness);
-		}
-	}
-
 	if (AnimInstanceProxy)
 	{
 		for (int32 Idx = 0; Idx < Num; Idx++)
 		{
-			AnimInstanceProxy->AnimDrawDebugLine(StartLocations[Idx], EndLocations[Idx], Color.ToFColor(true), false, -1.0f, Thickness, bDepthTest ? SDPG_World : SDPG_Foreground);
+			AnimInstanceProxy->AnimDrawDebugLine(Locations[Idx], Locations[Idx], Color.ToFColor(true), false, -1.0f, Thickness, bDepthTest ? SDPG_World : SDPG_Foreground);
 		}
 	}
 #endif
@@ -610,33 +553,58 @@ void FDebugDrawer::DrawLines(const TArrayView<const FVector> StartLocations, con
 		FVisualLogEntry* CurrentEntry = nullptr;
 		if (FVisualLogger::CheckVisualLogInputInternal(VisualLoggerObject, VisualLoggerCategory, VisualLoggerVerbosity, /*OutWorld=*/ nullptr, &CurrentEntry))
 		{
-			const uint16 IntegerThickness = Thickness == 0.0f ? 0 : FMath::CeilToInt(FMath::Max(2.0f * Thickness, 0.0f));
-			FVisualLogShapeElement Element(TEXT(""), Color.ToFColor(false), IntegerThickness, VisualLoggerCategory);
-			Element.Type = EVisualLoggerShapeElement::Segment;
-			Element.Verbosity = VisualLoggerVerbosity;
-			Element.Points.Reserve(2 * Num);
-			for (int32 Idx = 0; Idx < Num; Idx++)
+			if (Thickness == 0.0f)
 			{
-				Element.Points.Add(StartLocations[Idx]);
-				Element.Points.Add(EndLocations[Idx]);
+				FVisualLogShapeElement Element(TEXT(""), Color.ToFColor(false), 0, VisualLoggerCategory);
+				Element.Type = EVisualLoggerShapeElement::SinglePoint;
+				Element.Verbosity = VisualLoggerVerbosity;
+				Element.Points.Reserve(Num);
+				for (int32 Idx = 0; Idx < Num; Idx++)
+				{
+					Element.Points.Add(Locations[Idx]);
+				}
+				CurrentEntry->ElementsToDraw.Add(Element);
+			}
+			else
+			{
+				const uint16 IntegerThickness = Thickness == 0.0f ? 0 : FMath::CeilToInt(FMath::Max(2.0f * Thickness, 0.0f));
+				FVisualLogShapeElement Element(TEXT(""), Color.ToFColor(false), IntegerThickness, VisualLoggerCategory);
+				Element.Type = EVisualLoggerShapeElement::Segment;
+				Element.Verbosity = VisualLoggerVerbosity;
+				Element.Points.Reserve(2 * Num);
+				for (int32 Idx = 0; Idx < Num; Idx++)
+				{
+					Element.Points.Add(Locations[Idx]);
+					Element.Points.Add(Locations[Idx]);
+				}
+				CurrentEntry->ElementsToDraw.Add(Element);
 			}
-			CurrentEntry->ElementsToDraw.Add(Element);
 		}
-
+		
 		const UWorld* VisualLoggerWorld = VisualLoggerObject->GetWorld();
 
 		if (VisualLoggerWorld && bVisualLoggerDrawToScene && (bVisualLoggerDrawToSceneWhileRecording || !FVisualLogger::IsRecording()))
 		{
-			for (int32 Idx = 0; Idx < Num; Idx++)
+			if (Thickness == 0.0f)
+			{
+				for (int32 Idx = 0; Idx < Num; Idx++)
+				{
+					::DrawDebugPoint(VisualLoggerWorld, Locations[Idx], 1.0f, Color.ToFColor(true), false, -1.0f, bDepthTest ? 0 : 1);
+				}
+			}
+			else
 			{
-				::DrawDebugLine(VisualLoggerWorld, StartLocations[Idx], EndLocations[Idx], Color.ToFColor(true), false, -1.0f, bDepthTest ? 0 : 1, Thickness);
+				for (int32 Idx = 0; Idx < Num; Idx++)
+				{
+					::DrawDebugLine(VisualLoggerWorld, Locations[Idx], Locations[Idx], Color.ToFColor(true), false, -1.0f, bDepthTest ? 0 : 1, Thickness);
+				}
 			}
 		}
 	}
 #endif
 }
 
-void FDebugDrawer::DrawLinesPairs(const TArrayView<const TPair<FVector, FVector>> Segments, const FLinearColor& Color, const float Thickness, const bool bDepthTest) const
+void FDebugDrawer::DrawLines(const TArrayView<const TPair<FVector, FVector>> Segments, const FLinearColor& Color, const float Thickness, const bool bDepthTest) const
 {
 	const int32 Num = Segments.Num();
 
@@ -869,42 +837,183 @@ void UDrawDebugLibrary::AddToVectorHistoryArray(TArray<FVector>& InOutValues, co
 	}
 }
 
-void UDrawDebugLibrary::DrawDebugPoint(const FDebugDrawer& Drawer, const FVector& Location, const FLinearColor Color, const float Thickness, const bool bDepthTest)
+FDrawDebugPointStyle UDrawDebugLibrary::MakeDrawDebugPointStyleFromColor(const FLinearColor Color)
+{
+	FDrawDebugPointStyle Out;
+	Out.Color = Color;
+	return Out;
+}
+
+FDrawDebugLineStyle UDrawDebugLibrary::MakeDrawDebugLineStyleFromColor(const FLinearColor Color)
+{
+	FDrawDebugLineStyle Out;
+	Out.Color = Color;
+	return Out;
+}
+
+FDrawDebugPointStyle UDrawDebugLibrary::DrawDebugPointStyleFromLineStyle(const FDrawDebugLineStyle& LineStyle)
+{
+	FDrawDebugPointStyle Out;
+	Out.Color = LineStyle.Color;
+	Out.Thickness = LineStyle.Thickness;
+	return Out;
+}
+
+FDrawDebugLineStyle UDrawDebugLibrary::DrawDebugLineStyleFromPointStyle(const FDrawDebugPointStyle& PointStyle)
+{
+	FDrawDebugLineStyle Out;
+	Out.Color = PointStyle.Color;
+	Out.Thickness = PointStyle.Thickness;
+	return Out;
+}
+
+FLinearColor UDrawDebugLibrary::DrawDebugPointStyleColor(const FDrawDebugPointStyle& PointStyle)
+{
+	return PointStyle.Color;
+}
+
+FLinearColor UDrawDebugLibrary::DrawDebugLineStyleColor(const FDrawDebugLineStyle& LineStyle)
+{
+	return LineStyle.Color;
+}
+
+FDrawDebugLineStyle UDrawDebugLibrary::DrawDebugLineStyleWithColor(const FDrawDebugLineStyle& LineStyle, const FLinearColor Color)
+{
+	FDrawDebugLineStyle Out = LineStyle;
+	Out.Color = Color;
+	return Out;
+}
+
+FDrawDebugLineStyle UDrawDebugLibrary::DrawDebugLineStyleWithColorNoOpacity(const FDrawDebugLineStyle& LineStyle, const FLinearColor Color)
+{
+	FDrawDebugLineStyle Out = LineStyle;
+	Out.Color.R = Color.R;
+	Out.Color.G = Color.G;
+	Out.Color.B = Color.B;
+	return Out;
+}
+
+FDrawDebugLineStyle UDrawDebugLibrary::DrawDebugLineStyleWithThickness(const FDrawDebugLineStyle& LineStyle, const float Thickness)
 {
-	Drawer.DrawPoints({ Location }, Color, Thickness, bDepthTest);
+	FDrawDebugLineStyle Out = LineStyle;
+	Out.Thickness = Thickness;
+	return Out;
 }
 
-void UDrawDebugLibrary::DrawDebugPoints(const FDebugDrawer& Drawer, const TArray<FVector>& Locations, const FLinearColor Color, const float Thickness, const bool bDepthTest)
+FDrawDebugLineStyle UDrawDebugLibrary::DrawDebugLineStyleWithType(const FDrawDebugLineStyle& LineStyle, const EDrawDebugLineType LineType)
 {
-	Drawer.DrawPoints(Locations, Color, Thickness, bDepthTest);
+	FDrawDebugLineStyle Out = LineStyle;
+	Out.LineType = LineType;
+	return Out;
 }
 
-void UDrawDebugLibrary::DrawDebugPointsArrayView(const FDebugDrawer& Drawer, const TArrayView<const FVector> Locations, const FLinearColor Color, const float Thickness, const bool bDepthTest)
+void UDrawDebugLibrary::DrawDebugPoint(const FDebugDrawer& Drawer, const FVector& Location, const FDrawDebugPointStyle& PointStyle, const bool bDepthTest)
 {
-	Drawer.DrawPoints(Locations, Color, Thickness, bDepthTest);
+	DrawDebugPointsArrayView(Drawer, { Location }, PointStyle, bDepthTest);
 }
 
-void UDrawDebugLibrary::DrawDebugLine(const FDebugDrawer& Drawer, const FVector& StartLocation, const FVector& EndLocation, const FLinearColor Color, const float Thickness, const bool bDepthTest)
+void UDrawDebugLibrary::DrawDebugPoints(const FDebugDrawer& Drawer, const TArray<FVector>& Locations, const FDrawDebugPointStyle& PointStyle, const bool bDepthTest)
 {
-	Drawer.DrawLines({ StartLocation }, { EndLocation }, Color, Thickness, bDepthTest);
+	DrawDebugPointsArrayView(Drawer, Locations, PointStyle, bDepthTest);
 }
 
-void UDrawDebugLibrary::DrawDebugLines(const FDebugDrawer& Drawer, const TArray<FVector>& StartLocations, const TArray<FVector>& EndLocations, const FLinearColor Color, const float Thickness, const bool bDepthTest)
+void UDrawDebugLibrary::DrawDebugPointsArrayView(const FDebugDrawer& Drawer, const TArrayView<const FVector> Locations, const FDrawDebugPointStyle& PointStyle, const bool bDepthTest)
 {
-	Drawer.DrawLines(StartLocations, EndLocations, Color, Thickness, bDepthTest);
+	Drawer.DrawPoints(Locations, PointStyle.Color, PointStyle.Thickness, bDepthTest);
 }
 
-void UDrawDebugLibrary::DrawDebugLinesArrayView(const FDebugDrawer& Drawer, const TArrayView<const FVector> StartLocations, const TArrayView<const FVector> EndLocations, const FLinearColor Color, const float Thickness, const bool bDepthTest)
+void UDrawDebugLibrary::DrawDebugLine(const FDebugDrawer& Drawer, const FVector& StartLocation, const FVector& EndLocation, const FDrawDebugLineStyle& LineStyle, const bool bDepthTest)
 {
-	Drawer.DrawLines(StartLocations, EndLocations, Color, Thickness, bDepthTest);
+	DrawDebugLinesArrayView(Drawer, { StartLocation }, { EndLocation }, LineStyle, bDepthTest);
 }
 
-void UDrawDebugLibrary::DrawDebugLinesPairsArrayView(const FDebugDrawer& Drawer, const TArrayView<const TPair<FVector, FVector>> Segments, const FLinearColor Color, const float Thickness, const bool bDepthTest)
+void UDrawDebugLibrary::DrawDebugLines(const FDebugDrawer& Drawer, const TArray<FVector>& StartLocations, const TArray<FVector>& EndLocations, const FDrawDebugLineStyle& LineStyle, const bool bDepthTest)
 {
-	Drawer.DrawLinesPairs(Segments, Color, Thickness, bDepthTest);
+	DrawDebugLinesArrayView(Drawer, StartLocations, EndLocations, LineStyle, bDepthTest);
 }
 
-void UDrawDebugLibrary::DrawDebugTriangularBasePyramid(const FDebugDrawer& Drawer, const FVector& Location, const FRotator& Rotation, const FLinearColor Color, const float Thickness, const bool bDepthTest, const float Length, const float Width)
+void UDrawDebugLibrary::DrawDebugLinesArrayView(const FDebugDrawer& Drawer, const TArrayView<const FVector> StartLocations, const TArrayView<const FVector> EndLocations, const FDrawDebugLineStyle& LineStyle, const bool bDepthTest)
+{
+	TArray<TPair<FVector, FVector>, TInlineAllocator<64>> Segments;
+	const int32 Num = FMath::Min(StartLocations.Num(), EndLocations.Num());
+	Segments.SetNumUninitialized(Num);
+	for (int32 Idx = 0; Idx < Num; Idx++)
+	{
+		Segments[Idx] = { StartLocations[Idx], EndLocations[Idx] };
+	}
+
+	DrawDebugLinesPairsArrayView(Drawer, Segments, LineStyle, bDepthTest);
+}
+
+void UDrawDebugLibrary::DrawDebugLinesPairsArrayView(const FDebugDrawer& Drawer, const TArrayView<const TPair<FVector, FVector>> Segments, const FDrawDebugLineStyle& LineStyle, const bool bDepthTest)
+{
+	const float SafeThickness = FMath::Max(LineStyle.Thickness + 1.0f, 0.0f);
+
+	switch (LineStyle.LineType)
+	{
+
+	case EDrawDebugLineType::Solid:
+	{
+		Drawer.DrawLines(Segments, LineStyle.Color, LineStyle.Thickness, bDepthTest);
+		return;
+	}
+
+	case EDrawDebugLineType::Dotted:
+	{
+		const float SafeDotSpacing = FMath::Max(LineStyle.DotSpacing, 1.0f);
+		TArray<FVector> Dots;
+
+		float Dist = 0.0f;
+		const int32 Num = Segments.Num();
+		for (int32 Idx = 0; Idx < Num; Idx++)
+		{
+			const FVector Start = Segments[Idx].Key;
+			const FVector End = Segments[Idx].Value;
+			const FVector Direction = (End - Start).GetSafeNormal();
+			const float Length = (End - Start).Length();
+			for (; Dist < Length; Dist += SafeDotSpacing * SafeThickness)
+			{
+				Dots.Add(Start + Direction * Dist);
+			}
+			Dist = FMath::Fmod(Dist, Length);
+		}
+
+		Drawer.DrawPoints(Dots, LineStyle.Color, LineStyle.Thickness, bDepthTest);
+		return;
+	}
+
+	case EDrawDebugLineType::Dashed:
+	{
+		const float SafeDashSpacing = FMath::Max(LineStyle.DashSpacing, 1.0f);
+		TArray<TPair<FVector, FVector>, TInlineAllocator<64>> Dashes;
+
+		float Dist = 0.0f;
+		const int32 Num = Segments.Num();
+		for (int32 Idx = 0; Idx < Num; Idx++)
+		{
+			const FVector Start = Segments[Idx].Key;
+			const FVector End = Segments[Idx].Value;
+			const FVector Direction = (End - Start).GetSafeNormal();
+			const float Length = (End - Start).Length();
+			for (; Dist < Length; Dist += LineStyle.DashWidth + SafeDashSpacing * SafeThickness)
+			{
+				Dashes.Add({ 
+					Start + Direction * FMath::Clamp(Dist, 0.0f, Length), 
+					Start + Direction * FMath::Clamp(Dist + LineStyle.DashWidth, 0.0f, Length) });
+			}
+			Dist = FMath::Fmod(Dist, Length);
+		}
+
+		Drawer.DrawLines(Dashes, LineStyle.Color, LineStyle.Thickness, bDepthTest);
+		return;
+	}
+
+	default: checkNoEntry();
+	}
+}
+
+
+void UDrawDebugLibrary::DrawDebugTriangularBasePyramid(const FDebugDrawer& Drawer, const FVector& Location, const FRotator& Rotation, const FDrawDebugLineStyle& LineStyle, const bool bDepthTest, const float Length, const float Width)
 {
 	const FVector V0 = Location + Rotation.RotateVector(FVector(0, 0, Length/2));
 	const FVector V1 = Location + Rotation.RotateVector(FVector(-Width / 4 * FMath::Sqrt(3.0f), -Width / 2, -Length / 2));
@@ -918,10 +1027,10 @@ void UDrawDebugLibrary::DrawDebugTriangularBasePyramid(const FDebugDrawer& Drawe
 		{ V1, V2 },
 		{ V2, V3 },
 		{ V3, V1 },
-	}, Color, Thickness, bDepthTest);
+	}, LineStyle, bDepthTest);
 }
 
-void UDrawDebugLibrary::DrawDebugSquareBasePyramid(const FDebugDrawer& Drawer, const FVector& Location, const FRotator& Rotation, const FLinearColor Color, const float Thickness, const bool bDepthTest, const float Length, const float Width)
+void UDrawDebugLibrary::DrawDebugSquareBasePyramid(const FDebugDrawer& Drawer, const FVector& Location, const FRotator& Rotation, const FDrawDebugLineStyle& LineStyle, const bool bDepthTest, const float Length, const float Width)
 {
 	const FVector V0 = Location + Rotation.RotateVector(FVector(0, 0, Length / 2));
 	const FVector V1 = Location + Rotation.RotateVector(FVector(-Width / 2, -Width / 2, -Length / 2));
@@ -938,10 +1047,10 @@ void UDrawDebugLibrary::DrawDebugSquareBasePyramid(const FDebugDrawer& Drawer, c
 		{ V2, V3 },
 		{ V3, V4 },
 		{ V4, V1 },
-		}, Color, Thickness, bDepthTest);
+		}, LineStyle, bDepthTest);
 }
 
-void UDrawDebugLibrary::DrawDebugCone(const FDebugDrawer& Drawer, const FVector& Location, const FRotator& Rotation, const FLinearColor Color, const float Thickness, const bool bDepthTest, const float Length, const float Radius, const int32 Segments)
+void UDrawDebugLibrary::DrawDebugCone(const FDebugDrawer& Drawer, const FVector& Location, const FRotator& Rotation, const FDrawDebugLineStyle& LineStyle, const bool bDepthTest, const float Length, const float Radius, const int32 Segments)
 {
 	const int32 Num = FMath::Max(Segments, 2);
 
@@ -958,10 +1067,10 @@ void UDrawDebugLibrary::DrawDebugCone(const FDebugDrawer& Drawer, const FVector&
 		LineSegments.Add({ Location + Rotation.RotateVector(V0), Location + Rotation.RotateVector(V1) });
 	}
 
-	DrawDebugLinesPairsArrayView(Drawer, LineSegments, Color, Thickness, bDepthTest);
+	DrawDebugLinesPairsArrayView(Drawer, LineSegments, LineStyle, bDepthTest);
 }
 
-void UDrawDebugLibrary::DrawDebugArc(const FDebugDrawer& Drawer, const FVector& Location, const FRotator& Rotation, const float Angle, const FLinearColor Color, const float Thickness, const bool bDepthTest, const float Radius, const int32 Segments)
+void UDrawDebugLibrary::DrawDebugArc(const FDebugDrawer& Drawer, const FVector& Location, const FRotator& Rotation, const float Angle, const FDrawDebugLineStyle& LineStyle, const bool bDepthTest, const float Radius, const int32 Segments)
 {
 	TArray<TPair<FVector, FVector>, TInlineAllocator<17>> LineSegments;
 	LineSegments.Reserve(Segments);
@@ -975,25 +1084,25 @@ void UDrawDebugLibrary::DrawDebugArc(const FDebugDrawer& Drawer, const FVector&
 		LineSegments.Add({ Location0, Location1 });
 	}
 
-	DrawDebugLinesPairsArrayView(Drawer, LineSegments, Color, Thickness, bDepthTest);
+	DrawDebugLinesPairsArrayView(Drawer, LineSegments, LineStyle, bDepthTest);
 }
 
-void UDrawDebugLibrary::DrawDebugCircle(const FDebugDrawer& Drawer, const FVector& Location, const FRotator& Rotation, const FLinearColor Color, const float Thickness, const bool bDepthTest, const float Radius, const int32 Segments)
+void UDrawDebugLibrary::DrawDebugCircle(const FDebugDrawer& Drawer, const FVector& Location, const FRotator& Rotation, const FDrawDebugLineStyle& LineStyle, const bool bDepthTest, const float Radius, const int32 Segments)
 {
-	DrawDebugArc(Drawer, Location, Rotation, 360.0f, Color, Thickness, bDepthTest, Radius, FMath::Max(Segments, 3));
+	DrawDebugArc(Drawer, Location, Rotation, 360.0f, LineStyle, bDepthTest, Radius, FMath::Max(Segments, 3));
 }
 
-void UDrawDebugLibrary::DrawDebugCircleTick(const FDebugDrawer& Drawer, const FVector& Location, const FRotator& Rotation, const float Angle, const FLinearColor Color, const float Thickness, const bool bDepthTest, const float Radius, const float Length, bool bInside)
+void UDrawDebugLibrary::DrawDebugCircleTick(const FDebugDrawer& Drawer, const FVector& Location, const FRotator& Rotation, const float Angle, const FDrawDebugLineStyle& LineStyle, const bool bDepthTest, const float Radius, const float Length, bool bInside)
 {
-	DrawDebugCircleTicksArrayView(Drawer, Location, Rotation, { Angle }, Color, Thickness, bDepthTest, Radius, Length, bInside);
+	DrawDebugCircleTicksArrayView(Drawer, Location, Rotation, { Angle }, LineStyle, bDepthTest, Radius, Length, bInside);
 }
 
-void UDrawDebugLibrary::DrawDebugCircleTicks(const FDebugDrawer& Drawer, const FVector& Location, const FRotator& Rotation, const TArray<float>& Angles, const FLinearColor Color, const float Thickness, const bool bDepthTest, const float Radius, const float Length, bool bInside)
+void UDrawDebugLibrary::DrawDebugCircleTicks(const FDebugDrawer& Drawer, const FVector& Location, const FRotator& Rotation, const TArray<float>& Angles, const FDrawDebugLineStyle& LineStyle, const bool bDepthTest, const float Radius, const float Length, bool bInside)
 {
-	DrawDebugCircleTicksArrayView(Drawer, Location, Rotation, Angles, Color, Thickness, bDepthTest, Radius, Length, bInside);
+	DrawDebugCircleTicksArrayView(Drawer, Location, Rotation, Angles, LineStyle, bDepthTest, Radius, Length, bInside);
 }
 
-void UDrawDebugLibrary::DrawDebugCircleTicksArrayView(const FDebugDrawer& Drawer, const FVector& Location, const FRotator& Rotation, const TArrayView<const float> Angles, const FLinearColor Color, const float Thickness, const bool bDepthTest, const float Radius, const float Length, bool bInside)
+void UDrawDebugLibrary::DrawDebugCircleTicksArrayView(const FDebugDrawer& Drawer, const FVector& Location, const FRotator& Rotation, const TArrayView<const float> Angles, const FDrawDebugLineStyle& LineStyle, const bool bDepthTest, const float Radius, const float Length, bool bInside)
 {
 	const int32 TickNum = Angles.Num();
 
@@ -1009,15 +1118,15 @@ void UDrawDebugLibrary::DrawDebugCircleTicksArrayView(const FDebugDrawer& Drawer
 		LineSegments.Add({ Start, End });
 	}
 
-	DrawDebugLinesPairsArrayView(Drawer, LineSegments, Color, Thickness, bDepthTest);
+	DrawDebugLinesPairsArrayView(Drawer, LineSegments, LineStyle, bDepthTest);
 }
 
-void UDrawDebugLibrary::DrawDebugTriangle(const FDebugDrawer& Drawer, const FVector& Location, const FRotator& Rotation, const FLinearColor Color, const float Thickness, const bool bDepthTest, const float Radius)
+void UDrawDebugLibrary::DrawDebugTriangle(const FDebugDrawer& Drawer, const FVector& Location, const FRotator& Rotation, const FDrawDebugLineStyle& LineStyle, const bool bDepthTest, const float Radius)
 {
-	DrawDebugCircle(Drawer, Location, Rotation, Color, Thickness, bDepthTest, Radius, 3);
+	DrawDebugCircle(Drawer, Location, Rotation, LineStyle, bDepthTest, Radius, 3);
 }
 
-void UDrawDebugLibrary::DrawDebugSquare(const FDebugDrawer& Drawer, const FVector& Location, const FRotator& Rotation, const FLinearColor Color, const float Thickness, const bool bDepthTest, const float HalfLength)
+void UDrawDebugLibrary::DrawDebugSquare(const FDebugDrawer& Drawer, const FVector& Location, const FRotator& Rotation, const FDrawDebugLineStyle& LineStyle, const bool bDepthTest, const float HalfLength)
 {
 	const FVector V0 = Location + Rotation.RotateVector(HalfLength * FVector(+1, +1, 0));
 	const FVector V1 = Location + Rotation.RotateVector(HalfLength * FVector(-1, +1, 0));
@@ -1029,76 +1138,76 @@ void UDrawDebugLibrary::DrawDebugSquare(const FDebugDrawer& Drawer, const FVecto
 		{ V1, V2 },
 		{ V2, V3 },
 		{ V3, V0 },
-		}, Color, Thickness, bDepthTest);
+		}, LineStyle, bDepthTest);
 }
 
-void UDrawDebugLibrary::DrawDebugDiamond(const FDebugDrawer& Drawer, const FVector& Location, const FRotator& Rotation, const FLinearColor Color, const float Thickness, const bool bDepthTest, const float Radius)
+void UDrawDebugLibrary::DrawDebugDiamond(const FDebugDrawer& Drawer, const FVector& Location, const FRotator& Rotation, const FDrawDebugLineStyle& LineStyle, const bool bDepthTest, const float Radius)
 {
-	DrawDebugCircle(Drawer, Location, Rotation, Color, Thickness, bDepthTest, Radius, 4);
+	DrawDebugCircle(Drawer, Location, Rotation, LineStyle, bDepthTest, Radius, 4);
 }
 
-void UDrawDebugLibrary::DrawDebugPentagon(const FDebugDrawer& Drawer, const FVector& Location, const FRotator& Rotation, const FLinearColor Color, const float Thickness, const bool bDepthTest, const float Radius)
+void UDrawDebugLibrary::DrawDebugPentagon(const FDebugDrawer& Drawer, const FVector& Location, const FRotator& Rotation, const FDrawDebugLineStyle& LineStyle, const bool bDepthTest, const float Radius)
 {
-	DrawDebugCircle(Drawer, Location, Rotation, Color, Thickness, bDepthTest, Radius, 5);
+	DrawDebugCircle(Drawer, Location, Rotation, LineStyle, bDepthTest, Radius, 5);
 }
 
-void UDrawDebugLibrary::DrawDebugHexagon(const FDebugDrawer& Drawer, const FVector& Location, const FRotator& Rotation, const FLinearColor Color, const float Thickness, const bool bDepthTest, const float Radius)
+void UDrawDebugLibrary::DrawDebugHexagon(const FDebugDrawer& Drawer, const FVector& Location, const FRotator& Rotation, const FDrawDebugLineStyle& LineStyle, const bool bDepthTest, const float Radius)
 {
-	DrawDebugCircle(Drawer, Location, Rotation, Color, Thickness, bDepthTest, Radius, 6);
+	DrawDebugCircle(Drawer, Location, Rotation, LineStyle, bDepthTest, Radius, 6);
 }
 
-void UDrawDebugLibrary::DrawDebugHeptagon(const FDebugDrawer& Drawer, const FVector& Location, const FRotator& Rotation, const FLinearColor Color, const float Thickness, const bool bDepthTest, const float Radius)
+void UDrawDebugLibrary::DrawDebugHeptagon(const FDebugDrawer& Drawer, const FVector& Location, const FRotator& Rotation, const FDrawDebugLineStyle& LineStyle, const bool bDepthTest, const float Radius)
 {
-	DrawDebugCircle(Drawer, Location, Rotation, Color, Thickness, bDepthTest, Radius, 7);
+	DrawDebugCircle(Drawer, Location, Rotation, LineStyle, bDepthTest, Radius, 7);
 }
 
-void UDrawDebugLibrary::DrawDebugOctagon(const FDebugDrawer& Drawer, const FVector& Location, const FRotator& Rotation, const FLinearColor Color, const float Thickness, const bool bDepthTest, const float Radius)
+void UDrawDebugLibrary::DrawDebugOctagon(const FDebugDrawer& Drawer, const FVector& Location, const FRotator& Rotation, const FDrawDebugLineStyle& LineStyle, const bool bDepthTest, const float Radius)
 {
-	DrawDebugCircle(Drawer, Location, Rotation, Color, Thickness, bDepthTest, Radius, 8);
+	DrawDebugCircle(Drawer, Location, Rotation, LineStyle, bDepthTest, Radius, 8);
 }
 
-void UDrawDebugLibrary::DrawDebugNonagon(const FDebugDrawer& Drawer, const FVector& Location, const FRotator& Rotation, const FLinearColor Color, const float Thickness, const bool bDepthTest, const float Radius)
+void UDrawDebugLibrary::DrawDebugNonagon(const FDebugDrawer& Drawer, const FVector& Location, const FRotator& Rotation, const FDrawDebugLineStyle& LineStyle, const bool bDepthTest, const float Radius)
 {
-	DrawDebugCircle(Drawer, Location, Rotation, Color,Thickness, bDepthTest, Radius, 9);
+	DrawDebugCircle(Drawer, Location, Rotation, LineStyle, bDepthTest, Radius, 9);
 }
 
-void UDrawDebugLibrary::DrawDebugDecagon(const FDebugDrawer& Drawer, const FVector& Location, const FRotator& Rotation, const FLinearColor Color, const float Thickness, const bool bDepthTest, const float Radius)
+void UDrawDebugLibrary::DrawDebugDecagon(const FDebugDrawer& Drawer, const FVector& Location, const FRotator& Rotation, const FDrawDebugLineStyle& LineStyle, const bool bDepthTest, const float Radius)
 {
-	DrawDebugCircle(Drawer, Location, Rotation, Color, Thickness, bDepthTest, Radius, 10);
+	DrawDebugCircle(Drawer, Location, Rotation, LineStyle, bDepthTest, Radius, 10);
 }
 
-void UDrawDebugLibrary::DrawDebugRegularPolygon(const FDebugDrawer& Drawer, const FVector& Location, const FRotator& Rotation, const FLinearColor Color, const float Thickness, const bool bDepthTest, const float Radius, const int32 Sides)
+void UDrawDebugLibrary::DrawDebugRegularPolygon(const FDebugDrawer& Drawer, const FVector& Location, const FRotator& Rotation, const FDrawDebugLineStyle& LineStyle, const bool bDepthTest, const float Radius, const int32 Sides)
 {
-	DrawDebugCircle(Drawer, Location, Rotation, Color, Thickness, bDepthTest, Radius, Sides);
+	DrawDebugCircle(Drawer, Location, Rotation, LineStyle, bDepthTest, Radius, Sides);
 }
 
-void UDrawDebugLibrary::DrawDebugCircleOutline(const FDebugDrawer& Drawer, const FVector& Location, const FRotator& Rotation, const FLinearColor Color, const float Thickness, const bool bDepthTest, const float InnerRadius, const float OuterRadius, const int32 Segments)
+void UDrawDebugLibrary::DrawDebugCircleOutline(const FDebugDrawer& Drawer, const FVector& Location, const FRotator& Rotation, const FDrawDebugLineStyle& LineStyle, const bool bDepthTest, const float InnerRadius, const float OuterRadius, const int32 Segments)
 {
-	DrawDebugCircle(Drawer, Location, Rotation, Color, Thickness, bDepthTest, InnerRadius, Segments);
-	DrawDebugCircle(Drawer, Location, Rotation, Color, Thickness, bDepthTest, OuterRadius, Segments);
+	DrawDebugCircle(Drawer, Location, Rotation, LineStyle, bDepthTest, InnerRadius, Segments);
+	DrawDebugCircle(Drawer, Location, Rotation, LineStyle, bDepthTest, OuterRadius, Segments);
 }
 
-void UDrawDebugLibrary::DrawDebugLocator(const FDebugDrawer& Drawer, const FVector& Location, const FRotator& Rotation, const FLinearColor Color, const float Thickness, const bool bDepthTest, const float DrawRadius)
+void UDrawDebugLibrary::DrawDebugLocator(const FDebugDrawer& Drawer, const FVector& Location, const FRotator& Rotation, const FDrawDebugLineStyle& LineStyle, const bool bDepthTest, const float DrawRadius)
 {
 
 	DrawDebugLinesPairsArrayView(Drawer, {
 		{ Location - Rotation.RotateVector(FVector(DrawRadius, 0, 0)), Location + Rotation.RotateVector(FVector(DrawRadius, 0, 0)) },
 		{ Location - Rotation.RotateVector(FVector(0, DrawRadius, 0)), Location + Rotation.RotateVector(FVector(0, DrawRadius, 0)) },
 		{ Location - Rotation.RotateVector(FVector(0, 0, DrawRadius)), Location + Rotation.RotateVector(FVector(0, 0, DrawRadius)) },
-		}, Color, Thickness, bDepthTest);
+		}, LineStyle, bDepthTest);
 }
 
-void UDrawDebugLibrary::DrawDebugCrossLocator(const FDebugDrawer& Drawer, const FVector& Location, const FRotator& Rotation, const FLinearColor Color, const float Thickness, const bool bDepthTest, const float DrawRadius)
+void UDrawDebugLibrary::DrawDebugCrossLocator(const FDebugDrawer& Drawer, const FVector& Location, const FRotator& Rotation, const FDrawDebugLineStyle& LineStyle, const bool bDepthTest, const float DrawRadius)
 {
 	DrawDebugLinesPairsArrayView(Drawer, {
 		{ Location + Rotation.RotateVector(FVector(+DrawRadius, +DrawRadius, -DrawRadius)), Location + Rotation.RotateVector(FVector(-DrawRadius, -DrawRadius, +DrawRadius)) },
 		{ Location + Rotation.RotateVector(FVector(-DrawRadius, +DrawRadius, -DrawRadius)), Location + Rotation.RotateVector(FVector(+DrawRadius, -DrawRadius, +DrawRadius)) },
 		{ Location + Rotation.RotateVector(FVector(+DrawRadius, -DrawRadius, -DrawRadius)), Location + Rotation.RotateVector(FVector(-DrawRadius, +DrawRadius, +DrawRadius)) },
 		{ Location + Rotation.RotateVector(FVector(-DrawRadius, -DrawRadius, -DrawRadius)), Location + Rotation.RotateVector(FVector(+DrawRadius, +DrawRadius, +DrawRadius)) },
-		}, Color, Thickness, bDepthTest);
+		}, LineStyle, bDepthTest);
 }
 
-void UDrawDebugLibrary::DrawDebugBox(const FDebugDrawer& Drawer, const FVector& Location, const FRotator& Rotation, const FLinearColor Color, const float Thickness, const bool bDepthTest, const FVector HalfExtents)
+void UDrawDebugLibrary::DrawDebugBox(const FDebugDrawer& Drawer, const FVector& Location, const FRotator& Rotation, const FDrawDebugLineStyle& LineStyle, const bool bDepthTest, const FVector HalfExtents)
 {
 	const FVector C0 = Location + Rotation.RotateVector(FVector(+1, +1, -1) * HalfExtents);
 	const FVector C1 = Location + Rotation.RotateVector(FVector(+1, -1, -1) * HalfExtents);
@@ -1124,15 +1233,14 @@ void UDrawDebugLibrary::DrawDebugBox(const FDebugDrawer& Drawer, const FVector&
 		{ C1, C5 },
 		{ C2, C6 },
 		{ C3, C7 },
-		}, Color, Thickness, bDepthTest);
+		}, LineStyle, bDepthTest);
 }
 
 void UDrawDebugLibrary::DrawDebugSphere(
 	const FDebugDrawer& Drawer,
 	const FVector& Location,
 	const FRotator& Rotation,
-	const FLinearColor Color,
-	const float Thickness,
+	const FDrawDebugLineStyle& LineStyle,
 	const bool bDepthTest,
 	const float Radius,
 	const int32 Segments)
@@ -1161,10 +1269,10 @@ void UDrawDebugLibrary::DrawDebugSphere(
 		}
 	}
 
-	DrawDebugLinesPairsArrayView(Drawer, LineSegments, Color, Thickness, bDepthTest);
+	DrawDebugLinesPairsArrayView(Drawer, LineSegments, LineStyle, bDepthTest);
 }
 
-void UDrawDebugLibrary::DrawDebugSimpleSphere(const FDebugDrawer& Drawer, const FVector& Location, const FRotator& Rotation, const FLinearColor Color, const float Thickness, const bool bDepthTest, const float Radius, const int32 Segments)
+void UDrawDebugLibrary::DrawDebugSimpleSphere(const FDebugDrawer& Drawer, const FVector& Location, const FRotator& Rotation, const FDrawDebugLineStyle& LineStyle, const bool bDepthTest, const float Radius, const int32 Segments)
 {
 	TArray<TPair<FVector, FVector>, TInlineAllocator<39>> LineSegments;
 	LineSegments.Reserve(Segments*3);
@@ -1184,10 +1292,10 @@ void UDrawDebugLibrary::DrawDebugSimpleSphere(const FDebugDrawer& Drawer, const
 		LineSegments.Add({ Z0, Z1 });
 	}
 
-	DrawDebugLinesPairsArrayView(Drawer, LineSegments, Color, Thickness, bDepthTest);
+	DrawDebugLinesPairsArrayView(Drawer, LineSegments, LineStyle, bDepthTest);
 }
 
-void UDrawDebugLibrary::DrawDebugCapsule(const FDebugDrawer& Drawer, const FVector& Location, const FRotator& Rotation, const FLinearColor Color, const float Thickness, const bool bDepthTest, const float Radius, const float HalfLength, const int32 Segments)
+void UDrawDebugLibrary::DrawDebugCapsule(const FDebugDrawer& Drawer, const FVector& Location, const FRotator& Rotation, const FDrawDebugLineStyle& LineStyle, const bool bDepthTest, const float Radius, const float HalfLength, const int32 Segments)
 {
 	const int32 SegmentNum = (FMath::Max(Segments, 4) / 4) * 4;
 
@@ -1234,18 +1342,18 @@ void UDrawDebugLibrary::DrawDebugCapsule(const FDebugDrawer& Drawer, const FVect
 		LineSegments.Add({ E0, E1 });
 	}
 
-	DrawDebugLinesPairsArrayView(Drawer, LineSegments, Color, Thickness, bDepthTest);
+	DrawDebugLinesPairsArrayView(Drawer, LineSegments, LineStyle, bDepthTest);
 }
 
-void UDrawDebugLibrary::DrawDebugCapsuleLine(const FDebugDrawer& Drawer, const FVector& StartLocation, const FVector& EndLocation, const FLinearColor Color, const float Thickness, const bool bDepthTest, const float Radius, const int32 Segments)
+void UDrawDebugLibrary::DrawDebugCapsuleLine(const FDebugDrawer& Drawer, const FVector& StartLocation, const FVector& EndLocation, const FDrawDebugLineStyle& LineStyle, const bool bDepthTest, const float Radius, const int32 Segments)
 {
 	const FQuat Rotation = UE::DrawDebugLibrary::Private::FindOrientation(StartLocation, EndLocation);
 	const float Length = (EndLocation - StartLocation).Length();
 
-	DrawDebugCapsule(Drawer, StartLocation + Rotation.RotateVector(FVector(Length / 2.0f, 0.0f, 0.0f)), (Rotation * FQuat::MakeFromRotationVector(FVector(0.0f, UE_HALF_PI, 0.0f))).Rotator(), Color, Thickness, bDepthTest, Radius, Length / 2.0f, Segments);
+	DrawDebugCapsule(Drawer, StartLocation + Rotation.RotateVector(FVector(Length / 2.0f, 0.0f, 0.0f)), (Rotation * FQuat::MakeFromRotationVector(FVector(0.0f, UE_HALF_PI, 0.0f))).Rotator(), LineStyle, bDepthTest, Radius, Length / 2.0f, Segments);
 }
 
-void UDrawDebugLibrary::DrawDebugFrustum(const FDebugDrawer& Drawer, const FMatrix& FrustumToWorld, const FLinearColor Color, const float Thickness, const bool bDepthTest)
+void UDrawDebugLibrary::DrawDebugFrustum(const FDebugDrawer& Drawer, const FMatrix& FrustumToWorld, const FDrawDebugLineStyle& LineStyle, const bool bDepthTest)
 {
 	FVector Vertices[2][2][2];
 	for (int32 Z = 0; Z < 2; Z++)
@@ -1280,21 +1388,20 @@ void UDrawDebugLibrary::DrawDebugFrustum(const FDebugDrawer& Drawer, const FMatr
 		{ Vertices[0][1][0], Vertices[1][1][0] },
 		{ Vertices[0][0][1], Vertices[1][0][1] },
 		{ Vertices[0][1][1], Vertices[1][1][1] },
-		}, Color, Thickness, bDepthTest);
+		}, LineStyle, bDepthTest);
 }
 
 void UDrawDebugLibrary::DrawDebugArrow(
 	const FDebugDrawer& Drawer,
 	const FVector& StartLocation,
 	const FVector& EndLocation,
-	const FLinearColor Color,
-	const float Thickness,
+	const FDrawDebugLineStyle& LineStyle,
 	const bool bDepthTest,
 	const FDrawDebugArrowSettings& Settings)
 {
 	const FQuat Rotation = UE::DrawDebugLibrary::Private::FindOrientation(StartLocation, EndLocation);
 
-	return DrawDebugOrientedArrow(Drawer, StartLocation, Rotation.Rotator(), (EndLocation - StartLocation).Length(), Color, Thickness, bDepthTest, Settings);
+	return DrawDebugOrientedArrow(Drawer, StartLocation, Rotation.Rotator(), (EndLocation - StartLocation).Length(), LineStyle, bDepthTest, Settings);
 }
 
 void UDrawDebugLibrary::DrawDebugOrientedArrow(
@@ -1302,14 +1409,20 @@ void UDrawDebugLibrary::DrawDebugOrientedArrow(
 	const FVector& Location, 
 	const FRotator& Rotation, 
 	const float Length,
-	const FLinearColor Color,
-	const float Thickness,
+	const FDrawDebugLineStyle& LineStyle,
 	const bool bDepthTest,
 	const FDrawDebugArrowSettings& Settings)
 {
-	const FVector EndLocation = Location + Rotation.RotateVector(FVector(Length, 0.0f, 0.0f));
+	const FVector HeadStartLocation = Location;
+	const FVector HeadEndLocation = Location + Rotation.RotateVector(FVector(Length, 0.0f, 0.0f));
+
+	const FVector LineStartLocation = Settings.bArrowLineEndsAtStartHead ? 
+		Location + Rotation.RotateVector(FVector(Settings.ArrowHeadStartSize, 0.0f, 0.0f)) : HeadStartLocation;
+	
+	const FVector LineEndLocation = Settings.bArrowLineEndsAtEndHead ?
+		Location + Rotation.RotateVector(FVector(Length - Settings.ArrowHeadEndSize, 0.0f, 0.0f)) : HeadEndLocation;
 
-	DrawDebugLine(Drawer, Location, EndLocation, Color, Thickness, bDepthTest);
+	DrawDebugLine(Drawer, LineStartLocation, LineEndLocation, LineStyle, bDepthTest);
 
 	if (Settings.bArrowheadOnStart)
 	{
@@ -1318,47 +1431,89 @@ void UDrawDebugLibrary::DrawDebugOrientedArrow(
 		case EDrawDebugArrowHead::Simple:
 		{
 			DrawDebugLinesPairsArrayView(Drawer, {
-				{ Location, Location + Rotation.RotateVector(FVector(+Settings.ArrowHeadStartWidth, +Settings.ArrowHeadStartWidth, 0.0f)) },
-				{ Location, Location + Rotation.RotateVector(FVector(+Settings.ArrowHeadStartWidth, -Settings.ArrowHeadStartWidth, 0.0f)) },
-				}, Color, Thickness, bDepthTest);
+				{ HeadStartLocation, HeadStartLocation + Rotation.RotateVector(FVector(+Settings.ArrowHeadStartSize, +Settings.ArrowHeadStartSize, 0.0f)) },
+				{ HeadStartLocation, HeadStartLocation + Rotation.RotateVector(FVector(+Settings.ArrowHeadStartSize, -Settings.ArrowHeadStartSize, 0.0f)) },
+				}, LineStyle, bDepthTest);
 			break;
 		}
 
 		case EDrawDebugArrowHead::FourWay:
 		{
 			DrawDebugLinesPairsArrayView(Drawer, {
-				{ Location, Location + Rotation.RotateVector(FVector(+Settings.ArrowHeadStartWidth, +Settings.ArrowHeadStartWidth, 0)) },
-				{ Location, Location + Rotation.RotateVector(FVector(+Settings.ArrowHeadStartWidth, -Settings.ArrowHeadStartWidth, 0)) },
-				{ Location, Location + Rotation.RotateVector(FVector(+Settings.ArrowHeadStartWidth, 0, +Settings.ArrowHeadStartWidth)) },
-				{ Location, Location + Rotation.RotateVector(FVector(+Settings.ArrowHeadStartWidth, 0, -Settings.ArrowHeadStartWidth)) },
-				}, Color, Thickness, bDepthTest);
+				{ HeadStartLocation, HeadStartLocation + Rotation.RotateVector(FVector(+Settings.ArrowHeadStartSize, +Settings.ArrowHeadStartSize, 0)) },
+				{ HeadStartLocation, HeadStartLocation + Rotation.RotateVector(FVector(+Settings.ArrowHeadStartSize, -Settings.ArrowHeadStartSize, 0)) },
+				{ HeadStartLocation, HeadStartLocation + Rotation.RotateVector(FVector(+Settings.ArrowHeadStartSize, 0, +Settings.ArrowHeadStartSize)) },
+				{ HeadStartLocation, HeadStartLocation + Rotation.RotateVector(FVector(+Settings.ArrowHeadStartSize, 0, -Settings.ArrowHeadStartSize)) },
+				}, LineStyle, bDepthTest);
 			break;
 		}
 
-		case EDrawDebugArrowHead::Triangular:
+		case EDrawDebugArrowHead::TriangularPyramid:
 		{
 			DrawDebugTriangularBasePyramid(Drawer,
-				Location + Rotation.RotateVector(FVector(+Settings.ArrowHeadStartWidth / 2, 0, 0)),
+				HeadStartLocation + Rotation.RotateVector(FVector(+Settings.ArrowHeadStartSize / 2, 0, 0)),
 				(Rotation.Quaternion() * FQuat::MakeFromRotationVector(FVector(0, -UE_PI / 2, 0))).Rotator(),
-				Color, Thickness, bDepthTest, Settings.ArrowHeadStartWidth, Settings.ArrowHeadStartWidth);
+				LineStyle, bDepthTest, Settings.ArrowHeadStartSize, Settings.ArrowHeadStartSize);
 			break;
 		}
 
-		case EDrawDebugArrowHead::Square:
+		case EDrawDebugArrowHead::SquarePyramid:
 		{
 			DrawDebugSquareBasePyramid(Drawer,
-				Location + Rotation.RotateVector(FVector(+Settings.ArrowHeadStartWidth / 2, 0, 0)),
+				HeadStartLocation + Rotation.RotateVector(FVector(+Settings.ArrowHeadStartSize / 2, 0, 0)),
 				(Rotation.Quaternion() * FQuat::MakeFromRotationVector(FVector(0, -UE_PI / 2, 0))).Rotator(),
-				Color, Thickness, bDepthTest, Settings.ArrowHeadStartWidth, Settings.ArrowHeadStartWidth);
+				LineStyle, bDepthTest, Settings.ArrowHeadStartSize, Settings.ArrowHeadStartSize);
 			break;
 		}
 
 		case EDrawDebugArrowHead::Cone:
 		{
 			DrawDebugCone(Drawer,
-				Location + Rotation.RotateVector(FVector(+Settings.ArrowHeadStartWidth / 2, 0, 0)),
+				HeadStartLocation + Rotation.RotateVector(FVector(+Settings.ArrowHeadStartSize / 2, 0, 0)),
 				(Rotation.Quaternion() * FQuat::MakeFromRotationVector(FVector(0, -UE_PI / 2, 0))).Rotator(),
-				Color, Thickness, bDepthTest, Settings.ArrowHeadStartWidth, Settings.ArrowHeadStartWidth / 2, 7);
+				LineStyle, bDepthTest, Settings.ArrowHeadStartSize, Settings.ArrowHeadStartSize / 2, 7);
+			break;
+		}
+		
+		case EDrawDebugArrowHead::Triangle:
+		{
+			DrawDebugLinesPairsArrayView(Drawer, {
+				{ HeadStartLocation, HeadStartLocation + Rotation.RotateVector(FVector(+Settings.ArrowHeadStartSize, +Settings.ArrowHeadStartSize / 2.0f, 0.0f)) },
+				{ HeadStartLocation, HeadStartLocation + Rotation.RotateVector(FVector(+Settings.ArrowHeadStartSize, -Settings.ArrowHeadStartSize / 2.0f, 0.0f)) },
+				{ 
+					HeadStartLocation + Rotation.RotateVector(FVector(+Settings.ArrowHeadStartSize, -Settings.ArrowHeadStartSize / 2.0f, 0.0f)),
+					HeadStartLocation + Rotation.RotateVector(FVector(+Settings.ArrowHeadStartSize, +Settings.ArrowHeadStartSize / 2.0f, 0.0f)) },
+				}, LineStyle, bDepthTest);
+			break;
+		}
+
+		case EDrawDebugArrowHead::Square:
+		{
+			DrawDebugSquare(Drawer, HeadStartLocation, Rotation, LineStyle, bDepthTest, Settings.ArrowHeadStartSize);
+			break;
+		}
+
+		case EDrawDebugArrowHead::Diamond:
+		{
+			DrawDebugDiamond(Drawer, HeadStartLocation, Rotation, LineStyle, bDepthTest, Settings.ArrowHeadStartSize);
+			break;
+		}
+
+		case EDrawDebugArrowHead::Circle:
+		{
+			DrawDebugCircle(Drawer, HeadStartLocation, Rotation, LineStyle, bDepthTest, Settings.ArrowHeadStartSize);
+			break;
+		}
+
+		case EDrawDebugArrowHead::Sphere:
+		{
+			DrawDebugSimpleSphere(Drawer, HeadStartLocation, Rotation, LineStyle, bDepthTest, Settings.ArrowHeadStartSize);
+			break;
+		}
+
+		case EDrawDebugArrowHead::Box:
+		{
+			DrawDebugBox(Drawer, HeadStartLocation, Rotation, LineStyle, bDepthTest, FVector(Settings.ArrowHeadStartSize, Settings.ArrowHeadStartSize, Settings.ArrowHeadStartSize));
 			break;
 		}
 
@@ -1373,47 +1528,89 @@ void UDrawDebugLibrary::DrawDebugOrientedArrow(
 		case EDrawDebugArrowHead::Simple:
 		{
 			DrawDebugLinesPairsArrayView(Drawer, {
-				{ EndLocation, EndLocation + Rotation.RotateVector(FVector(-Settings.ArrowHeadEndWidth, +Settings.ArrowHeadEndWidth, 0.0f)) },
-				{ EndLocation, EndLocation + Rotation.RotateVector(FVector(-Settings.ArrowHeadEndWidth, -Settings.ArrowHeadEndWidth, 0.0f)) },
-				}, Color, Thickness, bDepthTest);
+				{ HeadEndLocation, HeadEndLocation + Rotation.RotateVector(FVector(-Settings.ArrowHeadEndSize, +Settings.ArrowHeadEndSize, 0.0f)) },
+				{ HeadEndLocation, HeadEndLocation + Rotation.RotateVector(FVector(-Settings.ArrowHeadEndSize, -Settings.ArrowHeadEndSize, 0.0f)) },
+				}, LineStyle, bDepthTest);
 			break;
 		}
 
 		case EDrawDebugArrowHead::FourWay:
 		{
 			DrawDebugLinesPairsArrayView(Drawer, {
-				{ EndLocation, EndLocation + Rotation.RotateVector(FVector(-Settings.ArrowHeadEndWidth, +Settings.ArrowHeadEndWidth, 0)) },
-				{ EndLocation, EndLocation + Rotation.RotateVector(FVector(-Settings.ArrowHeadEndWidth, -Settings.ArrowHeadEndWidth, 0)) },
-				{ EndLocation, EndLocation + Rotation.RotateVector(FVector(-Settings.ArrowHeadEndWidth, 0, +Settings.ArrowHeadEndWidth)) },
-				{ EndLocation, EndLocation + Rotation.RotateVector(FVector(-Settings.ArrowHeadEndWidth, 0, -Settings.ArrowHeadEndWidth)) },
-				}, Color, Thickness, bDepthTest);
+				{ HeadEndLocation, HeadEndLocation + Rotation.RotateVector(FVector(-Settings.ArrowHeadEndSize, +Settings.ArrowHeadEndSize, 0)) },
+				{ HeadEndLocation, HeadEndLocation + Rotation.RotateVector(FVector(-Settings.ArrowHeadEndSize, -Settings.ArrowHeadEndSize, 0)) },
+				{ HeadEndLocation, HeadEndLocation + Rotation.RotateVector(FVector(-Settings.ArrowHeadEndSize, 0, +Settings.ArrowHeadEndSize)) },
+				{ HeadEndLocation, HeadEndLocation + Rotation.RotateVector(FVector(-Settings.ArrowHeadEndSize, 0, -Settings.ArrowHeadEndSize)) },
+				}, LineStyle, bDepthTest);
 			break;
 		}
 
-		case EDrawDebugArrowHead::Triangular:
+		case EDrawDebugArrowHead::TriangularPyramid:
 		{
 			DrawDebugTriangularBasePyramid(Drawer,
-				EndLocation + Rotation.RotateVector(FVector(-Settings.ArrowHeadEndWidth / 2, 0, 0)),
+				HeadEndLocation + Rotation.RotateVector(FVector(-Settings.ArrowHeadEndSize / 2, 0, 0)),
 				(Rotation.Quaternion() * FQuat::MakeFromRotationVector(FVector(0, +UE_PI / 2, 0))).Rotator(),
-				Color, Thickness, bDepthTest, Settings.ArrowHeadEndWidth, Settings.ArrowHeadEndWidth);
+				LineStyle, bDepthTest, Settings.ArrowHeadEndSize, Settings.ArrowHeadEndSize);
 			break;
 		}
 
-		case EDrawDebugArrowHead::Square:
+		case EDrawDebugArrowHead::SquarePyramid:
 		{
 			DrawDebugSquareBasePyramid(Drawer,
-				EndLocation + Rotation.RotateVector(FVector(-Settings.ArrowHeadEndWidth / 2, 0, 0)),
+				HeadEndLocation + Rotation.RotateVector(FVector(-Settings.ArrowHeadEndSize / 2, 0, 0)),
 				(Rotation.Quaternion() * FQuat::MakeFromRotationVector(FVector(0, +UE_PI / 2, 0))).Rotator(),
-				Color, Thickness, bDepthTest, Settings.ArrowHeadEndWidth, Settings.ArrowHeadEndWidth);
+				LineStyle, bDepthTest, Settings.ArrowHeadEndSize, Settings.ArrowHeadEndSize);
 			break;
 		}
 
 		case EDrawDebugArrowHead::Cone:
 		{
 			DrawDebugCone(Drawer,
-				EndLocation + Rotation.RotateVector(FVector(-Settings.ArrowHeadEndWidth / 2, 0, 0)),
+				HeadEndLocation + Rotation.RotateVector(FVector(-Settings.ArrowHeadEndSize / 2, 0, 0)),
 				(Rotation.Quaternion() * FQuat::MakeFromRotationVector(FVector(0, +UE_PI / 2, 0))).Rotator(),
-				Color, Thickness, bDepthTest, Settings.ArrowHeadEndWidth, Settings.ArrowHeadEndWidth / 2, 7);
+				LineStyle, bDepthTest, Settings.ArrowHeadEndSize, Settings.ArrowHeadEndSize / 2, 7);
+			break;
+		}
+
+		case EDrawDebugArrowHead::Triangle:
+		{
+			DrawDebugLinesPairsArrayView(Drawer, {
+				{ HeadEndLocation, HeadEndLocation + Rotation.RotateVector(FVector(-Settings.ArrowHeadEndSize, +Settings.ArrowHeadEndSize / 2.0f, 0.0f)) },
+				{ HeadEndLocation, HeadEndLocation + Rotation.RotateVector(FVector(-Settings.ArrowHeadEndSize, -Settings.ArrowHeadEndSize / 2.0f, 0.0f)) },
+				{
+					HeadEndLocation + Rotation.RotateVector(FVector(-Settings.ArrowHeadEndSize, -Settings.ArrowHeadEndSize / 2.0f, 0.0f)),
+					HeadEndLocation + Rotation.RotateVector(FVector(-Settings.ArrowHeadEndSize, +Settings.ArrowHeadEndSize / 2.0f, 0.0f)) },
+				}, LineStyle, bDepthTest);
+			break;
+		}
+
+		case EDrawDebugArrowHead::Square:
+		{
+			DrawDebugSquare(Drawer, HeadEndLocation, Rotation, LineStyle, bDepthTest, Settings.ArrowHeadEndSize);
+			break;
+		}
+
+		case EDrawDebugArrowHead::Diamond:
+		{
+			DrawDebugDiamond(Drawer, HeadEndLocation, Rotation, LineStyle, bDepthTest, Settings.ArrowHeadEndSize);
+			break;
+		}
+
+		case EDrawDebugArrowHead::Circle:
+		{
+			DrawDebugCircle(Drawer, HeadEndLocation, Rotation, LineStyle, bDepthTest, Settings.ArrowHeadEndSize);
+			break;
+		}
+
+		case EDrawDebugArrowHead::Sphere:
+		{
+			DrawDebugSimpleSphere(Drawer, HeadEndLocation, Rotation, LineStyle, bDepthTest, Settings.ArrowHeadEndSize);
+			break;
+		}
+
+		case EDrawDebugArrowHead::Box:
+		{
+			DrawDebugBox(Drawer, HeadEndLocation, Rotation, LineStyle, bDepthTest, FVector(Settings.ArrowHeadEndSize, Settings.ArrowHeadEndSize, Settings.ArrowHeadEndSize));
 			break;
 		}
 
@@ -1422,7 +1619,7 @@ void UDrawDebugLibrary::DrawDebugOrientedArrow(
 	}
 }
 
-void UDrawDebugLibrary::DrawDebugGroundTargetArrow(const FDebugDrawer& Drawer, const FVector& Location, const FLinearColor Color, const float Thickness, const bool bDepthTest, const float Length, const float ArrowHeadSize, const float Radius, const int32 Segements)
+void UDrawDebugLibrary::DrawDebugGroundTargetArrow(const FDebugDrawer& Drawer, const FVector& Location, const FDrawDebugLineStyle& LineStyle, const bool bDepthTest, const float Length, const float ArrowHeadSize, const float Radius, const int32 Segements)
 {
 	FDrawDebugArrowSettings Settings;
 	Settings.ArrowHeadEndType = EDrawDebugArrowHead::FourWay;
@@ -1431,20 +1628,18 @@ void UDrawDebugLibrary::DrawDebugGroundTargetArrow(const FDebugDrawer& Drawer, c
 		Drawer,
 		Location + FVector(0.0f, 0.0f, Length),
 		Location,
-		Color,
-		Thickness,
+		LineStyle,
 		bDepthTest,
 		Settings);
 
-	DrawDebugCircle(Drawer, Location, FRotator::ZeroRotator, Color, Thickness, bDepthTest, Radius, Segements);
+	DrawDebugCircle(Drawer, Location, FRotator::ZeroRotator, LineStyle, bDepthTest, Radius, Segements);
 }
 
 void UDrawDebugLibrary::DrawDebugFlatArrow(
 	const FDebugDrawer& Drawer,
 	const FVector& Location,
 	const FRotator& Rotation,
-	const FLinearColor Color,
-	const float Thickness,
+	const FDrawDebugLineStyle& LineStyle,
 	const bool bDepthTest,
 	const float Length,
 	const float Width)
@@ -1465,17 +1660,17 @@ void UDrawDebugLibrary::DrawDebugFlatArrow(
 		{ V3, V5 },
 		{ V4, V6 },
 		{ V5, V6 },
-		}, Color, Thickness, bDepthTest);
+		}, LineStyle, bDepthTest);
 }
 
-void UDrawDebugLibrary::DrawDebugCircleArrow(const FDebugDrawer& Drawer, const FVector& Location, const FRotator& Rotation, const float Angle, const FLinearColor Color, const float Thickness, const bool bDepthTest, const float Radius, const float Length, const FDrawDebugArrowSettings& ArrowSettings)
+void UDrawDebugLibrary::DrawDebugCircleArrow(const FDebugDrawer& Drawer, const FVector& Location, const FRotator& Rotation, const float Angle, const FDrawDebugLineStyle& LineStyle, const bool bDepthTest, const float Radius, const float Length, const FDrawDebugArrowSettings& ArrowSettings)
 {
 	const FVector Start = Location + Rotation.RotateVector(Radius * FVector(FMath::Cos(FMath::DegreesToRadians(Angle)), FMath::Sin(FMath::DegreesToRadians(Angle)), 0.0f));
 	const FVector End = Location + Rotation.RotateVector((Radius + Length) * FVector(FMath::Cos(FMath::DegreesToRadians(Angle)), FMath::Sin(FMath::DegreesToRadians(Angle)), 0.0f));
-	DrawDebugArrow(Drawer, Start, End, Color, Thickness, bDepthTest, ArrowSettings);
+	DrawDebugArrow(Drawer, Start, End, LineStyle, bDepthTest, ArrowSettings);
 }
 
-void UDrawDebugLibrary::DrawDebugCatmullRomSplineStart(const FDebugDrawer& Drawer, const FVector& V0, const FVector& V1, const FVector& V2, const FLinearColor Color, const float Thickness, const bool bDepthTest, const bool bMonotonic, const int32 Segments)
+void UDrawDebugLibrary::DrawDebugCatmullRomSplineStart(const FDebugDrawer& Drawer, const FVector& V0, const FVector& V1, const FVector& V2, const FDrawDebugLineStyle& LineStyle, const bool bDepthTest, const bool bMonotonic, const int32 Segments)
 {
 	const int32 SegmentNum = FMath::Max(Segments, 2);
 
@@ -1505,10 +1700,10 @@ void UDrawDebugLibrary::DrawDebugCatmullRomSplineStart(const FDebugDrawer& Drawe
 		}
 	}
 
-	DrawDebugLinesPairsArrayView(Drawer, LineSegments, Color, Thickness, bDepthTest);
+	DrawDebugLinesPairsArrayView(Drawer, LineSegments, LineStyle, bDepthTest);
 }
 
-void UDrawDebugLibrary::DrawDebugCatmullRomSplineEnd(const FDebugDrawer& Drawer, const FVector& V0, const FVector& V1, const FVector& V2, const FLinearColor Color, const float Thickness, const bool bDepthTest, const bool bMonotonic, const int32 Segments)
+void UDrawDebugLibrary::DrawDebugCatmullRomSplineEnd(const FDebugDrawer& Drawer, const FVector& V0, const FVector& V1, const FVector& V2, const FDrawDebugLineStyle& LineStyle, const bool bDepthTest, const bool bMonotonic, const int32 Segments)
 {
 	const int32 SegmentNum = FMath::Max(Segments, 2);
 
@@ -1538,10 +1733,10 @@ void UDrawDebugLibrary::DrawDebugCatmullRomSplineEnd(const FDebugDrawer& Drawer,
 		}
 	}
 
-	DrawDebugLinesPairsArrayView(Drawer, LineSegments, Color, Thickness, bDepthTest);
+	DrawDebugLinesPairsArrayView(Drawer, LineSegments, LineStyle, bDepthTest);
 }
 
-void UDrawDebugLibrary::DrawDebugCatmullRomSplineSection(const FDebugDrawer& Drawer, const FVector& V0, const FVector& V1, const FVector& V2, const FVector& V3, const FLinearColor Color, const float Thickness, const bool bDepthTest, const bool bMonotonic, const int32 Segments)
+void UDrawDebugLibrary::DrawDebugCatmullRomSplineSection(const FDebugDrawer& Drawer, const FVector& V0, const FVector& V1, const FVector& V2, const FVector& V3, const FDrawDebugLineStyle& LineStyle, const bool bDepthTest, const bool bMonotonic, const int32 Segments)
 {
 	const int32 SegmentNum = FMath::Max(Segments, 2);
 
@@ -1571,68 +1766,68 @@ void UDrawDebugLibrary::DrawDebugCatmullRomSplineSection(const FDebugDrawer& Dra
 		}
 	}
 
-	DrawDebugLinesPairsArrayView(Drawer, LineSegments, Color, Thickness, bDepthTest);
+	DrawDebugLinesPairsArrayView(Drawer, LineSegments, LineStyle, bDepthTest);
 }
 
 
-void UDrawDebugLibrary::DrawDebugCatmullRomSpline(const FDebugDrawer& Drawer, const TArray<FVector>& Points, const FLinearColor Color, const float Thickness, const bool bDepthTest, const bool bMonotonic, const int32 Segments)
+void UDrawDebugLibrary::DrawDebugCatmullRomSpline(const FDebugDrawer& Drawer, const TArray<FVector>& Points, const FDrawDebugLineStyle& LineStyle, const bool bDepthTest, const bool bMonotonic, const int32 Segments)
 {
-	DrawDebugCatmullRomSplineArrayView(Drawer, Points, Color, Thickness, bDepthTest, bMonotonic, Segments);
+	DrawDebugCatmullRomSplineArrayView(Drawer, Points, LineStyle, bDepthTest, bMonotonic, Segments);
 }
 
-void UDrawDebugLibrary::DrawDebugCatmullRomSplineArrayView(const FDebugDrawer& Drawer, const TArrayView<const FVector> Points, const FLinearColor Color, const float Thickness, const bool bDepthTest, const bool bMonotonic, const int32 Segments)
+void UDrawDebugLibrary::DrawDebugCatmullRomSplineArrayView(const FDebugDrawer& Drawer, const TArrayView<const FVector> Points, const FDrawDebugLineStyle& LineStyle, const bool bDepthTest, const bool bMonotonic, const int32 Segments)
 {
 	const int32 Num = Points.Num();
 	
 	if (Num == 0) { return; }
 	if (Num == 1)
 	{
-		DrawDebugPoint(Drawer, Points[0], Color, Thickness, bDepthTest);
+		DrawDebugPoint(Drawer, Points[0], DrawDebugPointStyleFromLineStyle(LineStyle), bDepthTest);
 		return;
 	}
 
 	if (Num == 2)
 	{
-		DrawDebugLine(Drawer, Points[0], Points[1], Color, Thickness, bDepthTest);
+		DrawDebugLine(Drawer, Points[0], Points[1], LineStyle, bDepthTest);
 		return;
 	}
 
-	DrawDebugCatmullRomSplineStart(Drawer, Points[0], Points[1], Points[2], Color, Thickness, bDepthTest, bMonotonic, Segments);
+	DrawDebugCatmullRomSplineStart(Drawer, Points[0], Points[1], Points[2], LineStyle, bDepthTest, bMonotonic, Segments);
 
 	for (int32 Idx = 0; Idx < Num - 3; Idx++)
 	{
-		DrawDebugCatmullRomSplineSection(Drawer, Points[Idx + 0], Points[Idx + 1], Points[Idx + 2], Points[Idx + 3], Color, Thickness, bDepthTest, bMonotonic, Segments);
+		DrawDebugCatmullRomSplineSection(Drawer, Points[Idx + 0], Points[Idx + 1], Points[Idx + 2], Points[Idx + 3], LineStyle, bDepthTest, bMonotonic, Segments);
 	}
 
-	DrawDebugCatmullRomSplineEnd(Drawer, Points[Num - 3], Points[Num - 2], Points[Num - 1], Color, Thickness, bDepthTest, bMonotonic, Segments);
+	DrawDebugCatmullRomSplineEnd(Drawer, Points[Num - 3], Points[Num - 2], Points[Num - 1], LineStyle, bDepthTest, bMonotonic, Segments);
 }
 
 
 
-void UDrawDebugLibrary::DrawDebugAngle(const FDebugDrawer& Drawer, const FVector& Location, const FRotator& Rotation, const float Angle, const FLinearColor Color, const float Thickness, const bool bDepthTest, const float LineLength, const float AngleRadius, const int32 Segments)
+void UDrawDebugLibrary::DrawDebugAngle(const FDebugDrawer& Drawer, const FVector& Location, const FRotator& Rotation, const float Angle, const FDrawDebugLineStyle& LineStyle, const bool bDepthTest, const float LineLength, const float AngleRadius, const int32 Segments)
 {
-	DrawDebugArc(Drawer, Location, Rotation, Angle, Color, Thickness, bDepthTest, AngleRadius, Segments);
+	DrawDebugArc(Drawer, Location, Rotation, Angle, LineStyle, bDepthTest, AngleRadius, Segments);
 	DrawDebugLinesPairsArrayView(Drawer, {
 		{ Location, Location + Rotation.RotateVector(FVector(LineLength, 0, 0)) },
 		{ Location, Location + Rotation.RotateVector(LineLength * FVector(FMath::Cos(FMath::DegreesToRadians(Angle)), FMath::Sin(FMath::DegreesToRadians(Angle)), 0)) }
-		}, Color, Thickness, bDepthTest);
+		}, LineStyle, bDepthTest);
 }
 
-void UDrawDebugLibrary::DrawDebugLocation(const FDebugDrawer& Drawer, const FVector& Location, const FLinearColor Color, const float Thickness, const bool bDepthTest, const float DrawRadius, const int32 Segments)
+void UDrawDebugLibrary::DrawDebugLocation(const FDebugDrawer& Drawer, const FVector& Location, const FDrawDebugLineStyle& LineStyle, const bool bDepthTest, const float DrawRadius, const int32 Segments)
 {
-	DrawDebugSphere(Drawer, Location, FRotator::ZeroRotator, Color,Thickness, bDepthTest, DrawRadius, Segments);
+	DrawDebugSphere(Drawer, Location, FRotator::ZeroRotator, LineStyle, bDepthTest, DrawRadius, Segments);
 }
 
-void UDrawDebugLibrary::DrawDebugLocations(const FDebugDrawer& Drawer, const TArray<FVector>& Locations, const FLinearColor Color, const float Thickness, const bool bDepthTest, const float DrawRadius, const int32 Segments)
+void UDrawDebugLibrary::DrawDebugLocations(const FDebugDrawer& Drawer, const TArray<FVector>& Locations, const FDrawDebugLineStyle& LineStyle, const bool bDepthTest, const float DrawRadius, const int32 Segments)
 {
-	DrawDebugLocationsArrayView(Drawer, Locations, Color, Thickness, bDepthTest, DrawRadius, Segments);
+	DrawDebugLocationsArrayView(Drawer, Locations, LineStyle, bDepthTest, DrawRadius, Segments);
 }
 
-void UDrawDebugLibrary::DrawDebugLocationsArrayView(const FDebugDrawer& Drawer, const TArrayView<const FVector> Locations, const FLinearColor Color, const float Thickness, const bool bDepthTest, const float DrawRadius, const int32 Segments)
+void UDrawDebugLibrary::DrawDebugLocationsArrayView(const FDebugDrawer& Drawer, const TArrayView<const FVector> Locations, const FDrawDebugLineStyle& LineStyle, const bool bDepthTest, const float DrawRadius, const int32 Segments)
 {
 	for (const FVector& Location : Locations)
 	{
-		DrawDebugLocation(Drawer, Location, Color, Thickness, bDepthTest, DrawRadius, Segments);
+		DrawDebugLocation(Drawer, Location, LineStyle, bDepthTest, DrawRadius, Segments);
 	}
 }
 
@@ -1640,55 +1835,54 @@ void UDrawDebugLibrary::DrawDebugRotation(
 	const FDebugDrawer& Drawer,
 	const FVector& Location,
 	const FRotator& Rotation,
-	const float Opacity,
-	const float Thickness,
+	const FDrawDebugLineStyle& LineStyle,
 	const bool bDepthTest,
 	const float DrawRadius)
 {
-	DrawDebugLine(Drawer, Location, Location + Rotation.RotateVector(FVector(DrawRadius, 0.0f, 0.0f)), FLinearColor(1.0f, 0.0f, 0.0f, Opacity), Thickness, bDepthTest);
-	DrawDebugLine(Drawer, Location, Location + Rotation.RotateVector(FVector(0.0f, DrawRadius, 0.0f)), FLinearColor(0.0f, 1.0f, 0.0f, Opacity), Thickness, bDepthTest);
-	DrawDebugLine(Drawer, Location, Location + Rotation.RotateVector(FVector(0.0f, 0.0f, DrawRadius)), FLinearColor(0.0f, 0.0f, 1.0f, Opacity), Thickness, bDepthTest);
+	DrawDebugLine(Drawer, Location, Location + Rotation.RotateVector(FVector(DrawRadius, 0.0f, 0.0f)), DrawDebugLineStyleWithColorNoOpacity(LineStyle, FLinearColor::Red), bDepthTest);
+	DrawDebugLine(Drawer, Location, Location + Rotation.RotateVector(FVector(0.0f, DrawRadius, 0.0f)), DrawDebugLineStyleWithColorNoOpacity(LineStyle, FLinearColor::Green), bDepthTest);
+	DrawDebugLine(Drawer, Location, Location + Rotation.RotateVector(FVector(0.0f, 0.0f, DrawRadius)), DrawDebugLineStyleWithColorNoOpacity(LineStyle, FLinearColor::Blue), bDepthTest);
 }
 
-void UDrawDebugLibrary::DrawDebugDirection(const FDebugDrawer& Drawer, const FVector& Location, const FVector& Direction, const FLinearColor Color, const float Thickness, const bool bDepthTest, const float DrawArrowLength, const float ArrowHeadScale)
+void UDrawDebugLibrary::DrawDebugDirection(const FDebugDrawer& Drawer, const FVector& Location, const FVector& Direction, const FDrawDebugLineStyle& LineStyle, const bool bDepthTest, const float DrawArrowLength, const float ArrowHeadScale)
 {
-	DrawDebugArrow(Drawer, Location, Location + DrawArrowLength * Direction, Color, Thickness, bDepthTest);
+	DrawDebugArrow(Drawer, Location, Location + DrawArrowLength * Direction, LineStyle, bDepthTest);
 }
 
-void UDrawDebugLibrary::DrawDebugVelocity(const FDebugDrawer& Drawer, const FVector& Location, const FVector& Velocity, const FLinearColor Color, const float Thickness, const bool bDepthTest, const float DrawVelocityLineScale)
+void UDrawDebugLibrary::DrawDebugVelocity(const FDebugDrawer& Drawer, const FVector& Location, const FVector& Velocity, const FDrawDebugLineStyle& LineStyle, const bool bDepthTest, const float DrawVelocityLineScale)
 {
-	DrawDebugLine(Drawer, Location, Location + DrawVelocityLineScale * Velocity, Color, Thickness, bDepthTest);
+	DrawDebugLine(Drawer, Location, Location + DrawVelocityLineScale * Velocity, LineStyle, bDepthTest);
 }
 
-void UDrawDebugLibrary::DrawDebugTransform(const FDebugDrawer& Drawer, const FTransform& Transform, const float Opacity, const float Thickness, const bool bDepthTest, const float DrawRadius)
+void UDrawDebugLibrary::DrawDebugTransform(const FDebugDrawer& Drawer, const FTransform& Transform, const FDrawDebugLineStyle& LineStyle, const bool bDepthTest, const float DrawRadius)
 {
 	const FVector Location = Transform.GetLocation();
 	const FQuat Rotation = Transform.GetRotation();
 	const FVector Scale = Transform.GetScale3D();
 
-	DrawDebugLine(Drawer, Location, Location + Rotation.RotateVector(Scale * FVector(DrawRadius, 0, 0)), FLinearColor(1.0f, 0.0f, 0.0f, Opacity), Thickness, bDepthTest);
-	DrawDebugLine(Drawer, Location, Location + Rotation.RotateVector(Scale * FVector(0, DrawRadius, 0)), FLinearColor(0.0f, 1.0f, 0.0f, Opacity), Thickness, bDepthTest);
-	DrawDebugLine(Drawer, Location, Location + Rotation.RotateVector(Scale * FVector(0, 0, DrawRadius)), FLinearColor(0.0f, 0.0f, 1.0f, Opacity), Thickness, bDepthTest);
+	DrawDebugLine(Drawer, Location, Location + Rotation.RotateVector(Scale * FVector(DrawRadius, 0, 0)), DrawDebugLineStyleWithColorNoOpacity(LineStyle, FLinearColor::Red), bDepthTest);
+	DrawDebugLine(Drawer, Location, Location + Rotation.RotateVector(Scale * FVector(0, DrawRadius, 0)), DrawDebugLineStyleWithColorNoOpacity(LineStyle, FLinearColor::Green), bDepthTest);
+	DrawDebugLine(Drawer, Location, Location + Rotation.RotateVector(Scale * FVector(0, 0, DrawRadius)), DrawDebugLineStyleWithColorNoOpacity(LineStyle, FLinearColor::Blue), bDepthTest);
 }
 
-void UDrawDebugLibrary::DrawDebugEvent(const FDebugDrawer& Drawer, bool bTimeUntilEventKnown, const float TimeUntilEvent, const FVector& Location, const FLinearColor Color, const float Thickness, const bool bDepthTest, const float Size)
+void UDrawDebugLibrary::DrawDebugEvent(const FDebugDrawer& Drawer, bool bTimeUntilEventKnown, const float TimeUntilEvent, const FVector& Location, const FDrawDebugLineStyle& LineStyle, const bool bDepthTest, const float Size)
 {
 	DrawDebugLine(
 		Drawer,
 		Location + FVector(0, 0, +Size),
 		Location + FVector(0, 0, -Size),
-		Color, Thickness, bDepthTest);
+		LineStyle, bDepthTest);
 
-	DrawDebugSphere(Drawer, Location, FRotator::ZeroRotator, Color, Thickness, bDepthTest, Size * 0.1f, 8);
+	DrawDebugSphere(Drawer, Location, FRotator::ZeroRotator, LineStyle, bDepthTest, Size * 0.1f, 8);
 
 	if (bTimeUntilEventKnown && TimeUntilEvent < 1.0f && TimeUntilEvent > -1.0f)
 	{
-		DrawDebugSphere(Drawer, Location + FVector(0, 0, TimeUntilEvent * Size), FRotator::ZeroRotator, Color, Thickness, bDepthTest, Size * 0.25f, 8);
+		DrawDebugSphere(Drawer, Location + FVector(0, 0, TimeUntilEvent * Size), FRotator::ZeroRotator, LineStyle, bDepthTest, Size * 0.25f, 8);
 	}
 }
 
 
-void UDrawDebugLibrary::DrawDebugChair(const FDebugDrawer& Drawer, const FVector& Location, const FRotator& Rotation, const FLinearColor Color, const float Thickness, const bool bDepthTest, const FDrawDebugChairSettings& Settings)
+void UDrawDebugLibrary::DrawDebugChair(const FDebugDrawer& Drawer, const FVector& Location, const FRotator& Rotation, const FDrawDebugLineStyle& LineStyle, const bool bDepthTest, const FDrawDebugChairSettings& Settings)
 {
 	const float HalfWidth = Settings.Width / 2.0f;
 	const FVector F0 = Location + Rotation.RotateVector(FVector(+HalfWidth, +HalfWidth, 0));
@@ -1718,10 +1912,10 @@ void UDrawDebugLibrary::DrawDebugChair(const FDebugDrawer& Drawer, const FVector
 		{ S2, B2 },
 		{ S3, B3 },
 		{ B2, B3 },
-		}, Color, Thickness, bDepthTest);
+		}, LineStyle, bDepthTest);
 }
 
-void UDrawDebugLibrary::DrawDebugDoor(const FDebugDrawer& Drawer, const FVector& Location, const FRotator& Rotation, const FLinearColor Color, const float Thickness, const bool bDepthTest, const FDrawDebugDoorSettings& Settings)
+void UDrawDebugLibrary::DrawDebugDoor(const FDebugDrawer& Drawer, const FVector& Location, const FRotator& Rotation, const FDrawDebugLineStyle& LineStyle, const bool bDepthTest, const FDrawDebugDoorSettings& Settings)
 {
 	const float HalfWidth = Settings.Width / 2.0f;
 
@@ -1756,7 +1950,7 @@ void UDrawDebugLibrary::DrawDebugDoor(const FDebugDrawer& Drawer, const FVector&
 		{ I1, I2 },
 		{ I2, I3 },
 		{ I3, I0 },
-		}, Color, Thickness, bDepthTest);
+		}, LineStyle, bDepthTest);
 
 	if (Settings.bOuterHandle)
 	{
@@ -1765,7 +1959,7 @@ void UDrawDebugLibrary::DrawDebugDoor(const FDebugDrawer& Drawer, const FVector&
 			FVector(+Settings.HandleRadius + Settings.HandleExtension, Settings.Width - Settings.HandleOffset, Settings.HandleHeight);
 		const FVector OuterHandle = Location + Rotation.RotateVector(OpenRotation.RotateVector(LOuterHandle) + OpenOffset);
 
-		DrawDebugSimpleSphere(Drawer, OuterHandle, FRotator::ZeroRotator, Color, Thickness, bDepthTest, Settings.HandleRadius);
+		DrawDebugSimpleSphere(Drawer, OuterHandle, FRotator::ZeroRotator, DrawDebugLineStyleWithType(LineStyle, EDrawDebugLineType::Solid), bDepthTest, Settings.HandleRadius);
 	}
 
 	if (Settings.bInnerHandle)
@@ -1775,7 +1969,7 @@ void UDrawDebugLibrary::DrawDebugDoor(const FDebugDrawer& Drawer, const FVector&
 			FVector(-Settings.HandleRadius - Settings.HandleExtension, Settings.Width - Settings.HandleOffset, Settings.HandleHeight);
 		const FVector InnerHandle = Location + Rotation.RotateVector(OpenRotation.RotateVector(LInnerHandle) + OpenOffset);
 
-		DrawDebugSimpleSphere(Drawer, InnerHandle, FRotator::ZeroRotator, Color, Thickness, bDepthTest, Settings.HandleRadius);
+		DrawDebugSimpleSphere(Drawer, InnerHandle, FRotator::ZeroRotator, DrawDebugLineStyleWithType(LineStyle, EDrawDebugLineType::Solid), bDepthTest, Settings.HandleRadius);
 	}
 
 	if (Settings.bDrawEntryArrow)
@@ -1784,16 +1978,15 @@ void UDrawDebugLibrary::DrawDebugDoor(const FDebugDrawer& Drawer, const FVector&
 			Drawer,
 			Location + Rotation.RotateVector(FVector(-50.0f, 0.0f, 2.0f)),
 			Location + Rotation.RotateVector(FVector(-10.0f, 0.0f, 2.0f)),
-			Color,
-			Thickness,
+			LineStyle,
 			bDepthTest);
 	}
 }
 
-void UDrawDebugLibrary::DrawDebugCamera(const FDebugDrawer& Drawer, const FVector& Location, const FRotator& Rotation, const FLinearColor Color, const float Thickness, const bool bDepthTest, float Scale, float FOVDeg)
+void UDrawDebugLibrary::DrawDebugCamera(const FDebugDrawer& Drawer, const FVector& Location, const FRotator& Rotation, const FDrawDebugLineStyle& LineStyle, const bool bDepthTest, float Scale, float FOVDeg)
 {
-	DrawDebugRotation(Drawer, Location, Rotation, Color.A, Thickness,bDepthTest, Scale);
-	DrawDebugBox(Drawer, Location, Rotation, Color, Thickness, bDepthTest, FVector(Scale, Scale, Scale));
+	DrawDebugRotation(Drawer, Location, Rotation, LineStyle, bDepthTest, Scale);
+	DrawDebugBox(Drawer, Location, Rotation, LineStyle, bDepthTest, FVector(Scale, Scale, Scale));
 
 	// draw "lens" portion
 	const FRotationTranslationMatrix Axes(Rotation, Location);
@@ -1819,15 +2012,15 @@ void UDrawDebugLibrary::DrawDebugCamera(const FDebugDrawer& Drawer, const FVecto
 		{ C1, C2 },
 		{ C2, C3 },
 		{ C3, C0 },
-		}, Color, Thickness, bDepthTest);
+		}, LineStyle, bDepthTest);
 }
 
-void UDrawDebugLibrary::DrawDebugTrajectory(const FDebugDrawer& Drawer, const TArray<FVector>& Locations, const TArray<FVector>& Directions, const FTransform& RelativeTransform, const FLinearColor Color, const float Thickness, const bool bDepthTest, const float DrawArrowLength, const float PointRadius, const float ArrowHeadScale, const int32 Segments, const float VerticalOffset)
+void UDrawDebugLibrary::DrawDebugTrajectory(const FDebugDrawer& Drawer, const TArray<FVector>& Locations, const TArray<FVector>& Directions, const FTransform& RelativeTransform, const FDrawDebugLineStyle& LineStyle, const bool bDepthTest, const float DrawArrowLength, const float PointRadius, const float ArrowHeadScale, const int32 Segments, const float VerticalOffset)
 {
-	DrawDebugTrajectoryFromArrayViews(Drawer, Locations, Directions, RelativeTransform, Color, Thickness, bDepthTest, DrawArrowLength, PointRadius, ArrowHeadScale, Segments, VerticalOffset);
+	DrawDebugTrajectoryFromArrayViews(Drawer, Locations, Directions, RelativeTransform, LineStyle, bDepthTest, DrawArrowLength, PointRadius, ArrowHeadScale, Segments, VerticalOffset);
 }
 
-void UDrawDebugLibrary::DrawDebugTrajectoryFromArrayViews(const FDebugDrawer& Drawer, const TArrayView<const FVector> Locations, const TArrayView<const FVector> Directions, const FTransform& RelativeTransform, const FLinearColor Color, const float Thickness, const bool bDepthTest, const float DrawArrowLength, const float PointRadius, const float ArrowHeadScale, const int32 Segments, const float VerticalOffset)
+void UDrawDebugLibrary::DrawDebugTrajectoryFromArrayViews(const FDebugDrawer& Drawer, const TArrayView<const FVector> Locations, const TArrayView<const FVector> Directions, const FTransform& RelativeTransform, const FDrawDebugLineStyle& LineStyle, const bool bDepthTest, const float DrawArrowLength, const float PointRadius, const float ArrowHeadScale, const int32 Segments, const float VerticalOffset)
 {
 	const int32 SampleNum = FMath::Min(Locations.Num(), Directions.Num());
 
@@ -1836,22 +2029,22 @@ void UDrawDebugLibrary::DrawDebugTrajectoryFromArrayViews(const FDebugDrawer& Dr
 		const FVector Location = RelativeTransform.TransformPosition(Locations[SampleIdx]) + FVector(0.0f, 0.0f, VerticalOffset);
 		const FVector Direction = RelativeTransform.TransformVectorNoScale(Directions[SampleIdx]);
 
-		DrawDebugSimpleSphere(Drawer, Location, FRotator::ZeroRotator, Color, Thickness, bDepthTest, PointRadius, Segments);
-		DrawDebugArrow(Drawer, Location, Location + DrawArrowLength * Direction, Color, Thickness, bDepthTest);
+		DrawDebugSimpleSphere(Drawer, Location, FRotator::ZeroRotator, LineStyle, bDepthTest, PointRadius, Segments);
+		DrawDebugArrow(Drawer, Location, Location + DrawArrowLength * Direction, LineStyle, bDepthTest);
 
 		if (SampleIdx > 0)
 		{
 			const FVector PrevLocation = RelativeTransform.TransformPosition(Locations[SampleIdx - 1]) + FVector(0.0f, 0.0f, VerticalOffset);
-			DrawDebugLine(Drawer, Location, PrevLocation, Color, Thickness, bDepthTest);
+			DrawDebugLine(Drawer, Location, PrevLocation, LineStyle, bDepthTest);
 		}
 	}
 }
 
-void UDrawDebugLibrary::DrawDebugTransformTrajectory(const FDebugDrawer& Drawer, const FTransformTrajectory& TransformTrajectory, const FLinearColor Color, const float Thickness, const bool bDepthTest, const float Radius, const float VerticalOffset)
+void UDrawDebugLibrary::DrawDebugTransformTrajectory(const FDebugDrawer& Drawer, const FTransformTrajectory& TransformTrajectory, const FDrawDebugLineStyle& LineStyle, const bool bDepthTest, const float Radius, const float VerticalOffset)
 {
 
 	FTransformTrajectorySample TimeZeroSample = TransformTrajectory.GetSampleAtTime(0.0f);
-	DrawDebugSphere(Drawer, TimeZeroSample.Position + FVector(0.0f, 0.0f, VerticalOffset), FRotator::ZeroRotator, Color, Thickness, bDepthTest, Radius);
+	DrawDebugSphere(Drawer, TimeZeroSample.Position + FVector(0.0f, 0.0f, VerticalOffset), FRotator::ZeroRotator, LineStyle, bDepthTest, Radius);
 
 	const int32 SampleNum = TransformTrajectory.Samples.Num();
 	for (int32 SampleIdx = 0; SampleIdx < SampleNum; SampleIdx++)
@@ -1859,23 +2052,23 @@ void UDrawDebugLibrary::DrawDebugTransformTrajectory(const FDebugDrawer& Drawer,
 		const FVector Location = TransformTrajectory.Samples[SampleIdx].Position + FVector(0.0f, 0.0f, VerticalOffset);
 		const FQuat Rotation = TransformTrajectory.Samples[SampleIdx].Facing;
 
-		DrawDebugRotation(Drawer, Location, Rotation.Rotator(), Color.A, Thickness, bDepthTest, Radius);
-		DrawDebugString(Drawer, FString::Printf(TEXT("%5.2f"), TransformTrajectory.Samples[SampleIdx].TimeInSeconds), Location + FVector(0.0f, 0.0, 10.0f), FRotator::ZeroRotator, Color);
+		DrawDebugRotation(Drawer, Location, Rotation.Rotator(), LineStyle, bDepthTest, Radius);
+		DrawDebugString(Drawer, FString::Printf(TEXT("%5.2f"), TransformTrajectory.Samples[SampleIdx].TimeInSeconds), Location + FVector(0.0f, 0.0, 10.0f), FRotator::ZeroRotator, LineStyle);
 
 		if (SampleIdx > 0)
 		{
 			const FVector PrevLocation = TransformTrajectory.Samples[SampleIdx - 1].Position + FVector(0.0f, 0.0f, VerticalOffset);
-			DrawDebugLine(Drawer, Location, PrevLocation, Color, Thickness, bDepthTest);
+			DrawDebugLine(Drawer, Location, PrevLocation, LineStyle, bDepthTest);
 		}
 	}
 }
 
-void UDrawDebugLibrary::DrawDebugPose(const FDebugDrawer& Drawer, const TArray<FVector>& BoneLocations, const TArray<FVector>& BoneLinearVelocities, const FTransform& RelativeTransform, const FLinearColor Color, const float Thickness, const bool bDepthTest, const float DrawVelocityLineScale)
+void UDrawDebugLibrary::DrawDebugPose(const FDebugDrawer& Drawer, const TArray<FVector>& BoneLocations, const TArray<FVector>& BoneLinearVelocities, const FTransform& RelativeTransform, const FDrawDebugLineStyle& LineStyle, const bool bDepthTest, const float DrawVelocityLineScale)
 {
-	DrawDebugPoseFromArrayViews(Drawer, BoneLocations, BoneLinearVelocities, RelativeTransform, Color, Thickness, bDepthTest, DrawVelocityLineScale);
+	DrawDebugPoseFromArrayViews(Drawer, BoneLocations, BoneLinearVelocities, RelativeTransform, LineStyle, bDepthTest, DrawVelocityLineScale);
 }
 
-void UDrawDebugLibrary::DrawDebugPoseFromArrayViews(const FDebugDrawer& Drawer, const TArrayView<const FVector> BoneLocations, const TArrayView<const FVector> BoneLinearVelocities, const FTransform& RelativeTransform, const FLinearColor Color, const float Thickness, const bool bDepthTest, const float DrawVelocityLineScale)
+void UDrawDebugLibrary::DrawDebugPoseFromArrayViews(const FDebugDrawer& Drawer, const TArrayView<const FVector> BoneLocations, const TArrayView<const FVector> BoneLinearVelocities, const FTransform& RelativeTransform, const FDrawDebugLineStyle& LineStyle, const bool bDepthTest, const float DrawVelocityLineScale)
 {
 	const int32 BoneNum = FMath::Min(BoneLocations.Num(), BoneLinearVelocities.Num());
 
@@ -1884,24 +2077,24 @@ void UDrawDebugLibrary::DrawDebugPoseFromArrayViews(const FDebugDrawer& Drawer,
 		const FVector Location = RelativeTransform.TransformPosition(BoneLocations[BoneIdx]);
 		const FVector Velocity = RelativeTransform.TransformVector(BoneLinearVelocities[BoneIdx]);
 
-		DrawDebugSphere(Drawer, Location, FRotator::ZeroRotator, Color, Thickness, bDepthTest);
-		DrawDebugLine(Drawer, Location, Location + DrawVelocityLineScale * Velocity, Color, Thickness, bDepthTest);
+		DrawDebugSphere(Drawer, Location, FRotator::ZeroRotator, LineStyle, bDepthTest);
+		DrawDebugLine(Drawer, Location, Location + DrawVelocityLineScale * Velocity, LineStyle, bDepthTest);
 	}
 }
 
 
-void UDrawDebugLibrary::DrawDebugVelocities(const FDebugDrawer& Drawer, const TArray<FVector>& Locations, const TArray<FVector>& Velocities, const FLinearColor Color, const float Thickness, const bool bDepthTest, const float DrawVelocityLineScale)
+void UDrawDebugLibrary::DrawDebugVelocities(const FDebugDrawer& Drawer, const TArray<FVector>& Locations, const TArray<FVector>& Velocities, const FDrawDebugLineStyle& LineStyle, const bool bDepthTest, const float DrawVelocityLineScale)
 {
-	DrawDebugVelocitiesArrayView(Drawer, Locations, Velocities, Color, Thickness, bDepthTest, DrawVelocityLineScale);
+	DrawDebugVelocitiesArrayView(Drawer, Locations, Velocities, LineStyle, bDepthTest, DrawVelocityLineScale);
 }
 
-void UDrawDebugLibrary::DrawDebugVelocitiesArrayView(const FDebugDrawer& Drawer, const TArrayView<const FVector> Locations, const TArrayView<const FVector> Velocities, const FLinearColor Color, const float Thickness, const bool bDepthTest, const float DrawVelocityLineScale)
+void UDrawDebugLibrary::DrawDebugVelocitiesArrayView(const FDebugDrawer& Drawer, const TArrayView<const FVector> Locations, const TArrayView<const FVector> Velocities, const FDrawDebugLineStyle& LineStyle, const bool bDepthTest, const float DrawVelocityLineScale)
 {
 	const int32 VelocityNum = FMath::Min(Locations.Num(), Velocities.Num());
 
 	for (int32 VelocityIdx = 0; VelocityIdx < VelocityNum; VelocityIdx++)
 	{
-		DrawDebugVelocity(Drawer, Locations[VelocityIdx], Velocities[VelocityIdx], Color, Thickness, bDepthTest, DrawVelocityLineScale);
+		DrawDebugVelocity(Drawer, Locations[VelocityIdx], Velocities[VelocityIdx], LineStyle, bDepthTest, DrawVelocityLineScale);
 	}
 }
 
@@ -1911,8 +2104,7 @@ void UDrawDebugLibrary::DrawDebugTransformsArrayView(
 	const TArrayView<const FVector> Locations,
 	const TArrayView<const FQuat4f> Rotations,
 	const TArrayView<const FVector3f> Scales,
-	const float Opacity,
-	const float Thickness,
+	const FDrawDebugLineStyle& LineStyle,
 	const bool bDepthTest,
 	const float DrawRadius)
 {
@@ -1930,8 +2122,7 @@ void UDrawDebugLibrary::DrawDebugTransformsArrayView(
 		DrawDebugTransform(
 			Drawer,
 			Transform,
-			Opacity,
-			Thickness,
+			LineStyle,
 			bDepthTest,
 			DrawRadius);
 	}
@@ -1941,8 +2132,7 @@ void UDrawDebugLibrary::DrawDebugRotationsQuatArrayView(
 	const FDebugDrawer& Drawer,
 	const TArrayView<const FVector> Locations,
 	const TArrayView<const FQuat4f> Rotations,
-	const float Opacity,
-	const float Thickness,
+	const FDrawDebugLineStyle& LineStyle,
 	const bool bDepthTest,
 	const float DrawRadius)
 {
@@ -1954,8 +2144,7 @@ void UDrawDebugLibrary::DrawDebugRotationsQuatArrayView(
 			Drawer,
 			Locations[Idx],
 			(FRotator)Rotations[Idx].Rotator(),
-			Opacity,
-			Thickness,
+			LineStyle,
 			bDepthTest,
 			DrawRadius);
 	}
@@ -1971,17 +2160,23 @@ float UDrawDebugLibrary::GetDefaultBoneRadius()
 	return 1.0f;
 }
 
-void UDrawDebugLibrary::DrawDebugBone(const FDebugDrawer& Drawer, const FVector& Location, const FRotator& Rotation, const FLinearColor Color, const float Thickness, const bool bDepthTest, const float Radius, const int32 Segments, const bool bDrawTransform)
+void UDrawDebugLibrary::DrawDebugBone(const FDebugDrawer& Drawer, const FVector& Location, const FRotator& Rotation, const FLinearColor Color, const bool bDepthTest, const float Radius, const int32 Segments, const bool bDrawTransform)
 {
-	DrawDebugSimpleSphere(Drawer, Location, Rotation, Color, Thickness, bDepthTest, Radius, Segments);
+	FDrawDebugLineStyle LineStyle;
+	LineStyle.Color = Color;
+
+	DrawDebugSimpleSphere(Drawer, Location, Rotation, LineStyle, bDepthTest, Radius, Segments);
 	if (bDrawTransform)
 	{
-		DrawDebugRotation(Drawer, Location, Rotation, Color.A, Thickness, bDepthTest, Radius);
+		DrawDebugRotation(Drawer, Location, Rotation, LineStyle, bDepthTest, Radius);
 	}
 }
 
-void UDrawDebugLibrary::DrawDebugBoneLink(const FDebugDrawer& Drawer, const FVector& ChildLocation, const FVector& ParentLocation, const FRotator& ParentRotation, const FLinearColor Color, const float Thickness, const bool bDepthTest, const float Radius)
+void UDrawDebugLibrary::DrawDebugBoneLink(const FDebugDrawer& Drawer, const FVector& ChildLocation, const FVector& ParentLocation, const FRotator& ParentRotation, const FLinearColor Color, const bool bDepthTest, const float Radius)
 {
+	FDrawDebugLineStyle LineStyle;
+	LineStyle.Color = Color;
+
 	const FVector HalfWayLocation = (ChildLocation - ParentLocation) / 2.0f + ParentLocation;
 	const float Length = (ChildLocation - ParentLocation).Length() - 2.0f * Radius;
 
@@ -1991,7 +2186,7 @@ void UDrawDebugLibrary::DrawDebugBoneLink(const FDebugDrawer& Drawer, const FVec
 		const FVector LocalPosition = ParentRotation.UnrotateVector(ChildLocation - ParentLocation);
 		const FQuat LocalRotation = ParentRotation.Quaternion() * FQuat::FindBetween(FVector::UpVector, LocalPosition);
 
-		DrawDebugSquareBasePyramid(Drawer, HalfWayLocation, LocalRotation.Rotator(), Color, Thickness, bDepthTest, Length, Radius);
+		DrawDebugSquareBasePyramid(Drawer, HalfWayLocation, LocalRotation.Rotator(), LineStyle, bDepthTest, Length, Radius);
 	}
 }
 
@@ -1999,7 +2194,6 @@ void UDrawDebugLibrary::DrawDebugSkeletonFromSkinnedMeshComponent(
 	const FDebugDrawer& Drawer,
 	const USkinnedMeshComponent* SkinnedMeshComponent,
 	const FLinearColor Color,
-	const float Thickness,
 	const bool bDepthTest,
 	const FDrawDebugSkeletonSettings& Settings)
 {
@@ -2011,6 +2205,10 @@ void UDrawDebugLibrary::DrawDebugSkeletonFromSkinnedMeshComponent(
 
 	if (Settings.bDrawSimpleSkeleton)
 	{
+		FDrawDebugLineStyle LineStyle;
+		LineStyle.Color = Color;
+		LineStyle.Thickness = Settings.BoneRadius;
+
 		TArray<TPair<FVector, FVector>> LineSegments;
 		LineSegments.Reserve(BoneNum);
 
@@ -2028,11 +2226,11 @@ void UDrawDebugLibrary::DrawDebugSkeletonFromSkinnedMeshComponent(
 			}
 			else if (Settings.bDrawRoot && !(RootTransform.GetLocation() - BoneLocation).IsNearlyZero())
 			{
-				DrawDebugLine(Drawer, BoneLocation, RootTransform.GetLocation(), Settings.RootBoneColor, Settings.BoneRadius, bDepthTest);
+				DrawDebugLine(Drawer, BoneLocation, RootTransform.GetLocation(), LineStyle, bDepthTest);
 			}
 		}
 
-		DrawDebugLinesPairsArrayView(Drawer, LineSegments, Color, Settings.BoneRadius, bDepthTest);
+		DrawDebugLinesPairsArrayView(Drawer, LineSegments, LineStyle, bDepthTest);
 	}
 	else
 	{
@@ -2041,19 +2239,19 @@ void UDrawDebugLibrary::DrawDebugSkeletonFromSkinnedMeshComponent(
 			const FName BoneName = SkinnedMeshComponent->GetBoneName(BoneIdx);
 			const FTransform BoneTransform = SkinnedMeshComponent->GetBoneTransform(BoneName, ERelativeTransformSpace::RTS_World);
 
-			DrawDebugBone(Drawer, BoneTransform.GetLocation(), BoneTransform.GetRotation().Rotator(), Color, Thickness, bDepthTest, Settings.BoneRadius, Settings.BoneSegmentNum, Settings.bDrawTransforms);
+			DrawDebugBone(Drawer, BoneTransform.GetLocation(), BoneTransform.GetRotation().Rotator(), Color, bDepthTest, Settings.BoneRadius, Settings.BoneSegmentNum, Settings.bDrawTransforms);
 
 			const FName ParentName = SkinnedMeshComponent->GetParentBone(BoneName);
 			if (ParentName != NAME_None)
 			{
 				const FTransform ParentTransform = SkinnedMeshComponent->GetBoneTransform(ParentName, ERelativeTransformSpace::RTS_World);
 
-				DrawDebugBoneLink(Drawer, BoneTransform.GetLocation(), ParentTransform.GetLocation(), ParentTransform.GetRotation().Rotator(), Color, Thickness, bDepthTest, Settings.BoneRadius);
+				DrawDebugBoneLink(Drawer, BoneTransform.GetLocation(), ParentTransform.GetLocation(), ParentTransform.GetRotation().Rotator(), Color, bDepthTest, Settings.BoneRadius);
 			}
 			else if (Settings.bDrawRoot && !(RootTransform.GetLocation() - BoneTransform.GetLocation()).IsNearlyZero())
 			{
-				DrawDebugBone(Drawer, RootTransform.GetLocation(), RootTransform.GetRotation().Rotator(), Settings.RootBoneColor, Thickness, bDepthTest, Settings.BoneRadius, Settings.BoneSegmentNum, Settings.bDrawTransforms);
-				DrawDebugBoneLink(Drawer, BoneTransform.GetLocation(), RootTransform.GetLocation(), RootTransform.GetRotation().Rotator(), Settings.RootBoneColor, Thickness, bDepthTest, Settings.BoneRadius);
+				DrawDebugBone(Drawer, RootTransform.GetLocation(), RootTransform.GetRotation().Rotator(), Settings.RootBoneColor, bDepthTest, Settings.BoneRadius, Settings.BoneSegmentNum, Settings.bDrawTransforms);
+				DrawDebugBoneLink(Drawer, BoneTransform.GetLocation(), RootTransform.GetLocation(), RootTransform.GetRotation().Rotator(), Settings.RootBoneColor, bDepthTest, Settings.BoneRadius);
 			}
 		}
 	}
@@ -2068,12 +2266,15 @@ void UDrawDebugLibrary::DrawDebugSkeletonArrayView(
 	const TArrayView<const int32> BoneIndices,
 	const TArrayView<const int32> BoneParents,
 	const FLinearColor Color,
-	const float Thickness,
 	const bool bDepthTest,
 	const FDrawDebugSkeletonSettings& Settings)
 {
 	if (Settings.bDrawSimpleSkeleton)
 	{
+		FDrawDebugLineStyle LineStyle;
+		LineStyle.Color = Color;
+		LineStyle.Thickness = Settings.BoneRadius;
+
 		TArray<TPair<FVector, FVector>> LineSegments;
 		LineSegments.Reserve(BoneIndices.Num());
 
@@ -2087,29 +2288,29 @@ void UDrawDebugLibrary::DrawDebugSkeletonArrayView(
 			}
 			else if (Settings.bDrawRoot && !(RootLocation - BoneLocations[BoneIdx]).IsNearlyZero())
 			{
-				DrawDebugLine(Drawer, BoneLocations[BoneIdx], RootLocation, Settings.RootBoneColor, Settings.BoneRadius, bDepthTest);
+				DrawDebugLine(Drawer, BoneLocations[BoneIdx], RootLocation, LineStyle, bDepthTest);
 			}
 		}
 
-		DrawDebugLinesPairsArrayView(Drawer, LineSegments, Color, Settings.BoneRadius, bDepthTest);
+		DrawDebugLinesPairsArrayView(Drawer, LineSegments, LineStyle, bDepthTest);
 	}
 	else
 	{
 
 		for (int32 BoneIdx : BoneIndices)
 		{
-			DrawDebugBone(Drawer, BoneLocations[BoneIdx], (FRotator)BoneRotations[BoneIdx].Rotator(), Color, Thickness, bDepthTest, Settings.BoneRadius, Settings.BoneSegmentNum, Settings.bDrawTransforms);
+			DrawDebugBone(Drawer, BoneLocations[BoneIdx], (FRotator)BoneRotations[BoneIdx].Rotator(), Color, bDepthTest, Settings.BoneRadius, Settings.BoneSegmentNum, Settings.bDrawTransforms);
 
 			const int32 ParentIdx = BoneParents[BoneIdx];
 
 			if (ParentIdx != INDEX_NONE)
 			{
-				DrawDebugBoneLink(Drawer, BoneLocations[BoneIdx], BoneLocations[ParentIdx], (FRotator)BoneRotations[ParentIdx].Rotator(), Color, Thickness, bDepthTest, Settings.BoneRadius);
+				DrawDebugBoneLink(Drawer, BoneLocations[BoneIdx], BoneLocations[ParentIdx], (FRotator)BoneRotations[ParentIdx].Rotator(), Color, bDepthTest, Settings.BoneRadius);
 			}
 			else if (Settings.bDrawRoot && !(RootLocation - BoneLocations[BoneIdx]).IsNearlyZero())
 			{
-				DrawDebugBone(Drawer, RootLocation, (FRotator)RootRotation.Rotator(), Settings.RootBoneColor, Thickness, bDepthTest, Settings.BoneRadius, Settings.BoneSegmentNum, Settings.bDrawTransforms);
-				DrawDebugBoneLink(Drawer, BoneLocations[BoneIdx], RootLocation, (FRotator)RootRotation.Rotator(), Settings.RootBoneColor, Thickness, bDepthTest, Settings.BoneRadius);
+				DrawDebugBone(Drawer, RootLocation, (FRotator)RootRotation.Rotator(), Settings.RootBoneColor, bDepthTest, Settings.BoneRadius, Settings.BoneSegmentNum, Settings.bDrawTransforms);
+				DrawDebugBoneLink(Drawer, BoneLocations[BoneIdx], RootLocation, (FRotator)RootRotation.Rotator(), Settings.RootBoneColor, bDepthTest, Settings.BoneRadius);
 			}
 		}
 	}
@@ -2120,11 +2321,15 @@ void UDrawDebugLibrary::DrawDebugBoneVelocitiesArrayView(
 	const TArrayView<const FVector> BoneLocations,
 	const TArrayView<const FVector3f> BoneVelocities,
 	const TArrayView<const int32> BoneIndices,
-	const FLinearColor& Color,
+	const FLinearColor Color,
 	const float Thickness,
 	const bool bDepthTest,
 	const float Scale)
 {
+	FDrawDebugLineStyle LineStyle;
+	LineStyle.Color = Color;
+	LineStyle.Thickness = Thickness;
+
 	TArray<TPair<FVector, FVector>> LineSegments;
 	LineSegments.Reserve(BoneIndices.Num());
 
@@ -2133,7 +2338,7 @@ void UDrawDebugLibrary::DrawDebugBoneVelocitiesArrayView(
 		LineSegments.Add({ BoneLocations[BoneIdx], BoneLocations[BoneIdx] + Scale * (FVector)BoneVelocities[BoneIdx] });
 	}
 
-	DrawDebugLinesPairsArrayView(Drawer, LineSegments, Color, Thickness, bDepthTest);
+	DrawDebugLinesPairsArrayView(Drawer, LineSegments, LineStyle, bDepthTest);
 }
 
 void UDrawDebugLibrary::DrawDebugRangeTrajectoryArrayView(
@@ -2149,6 +2354,18 @@ void UDrawDebugLibrary::DrawDebugRangeTrajectoryArrayView(
 	const int32 FrameNum = TrajectoryLocations.Num();
 	if (FrameNum == 0) { return; }
 
+	FDrawDebugPointStyle PointStyle;
+	PointStyle.Color = FColor::Black.WithAlpha(64);
+	PointStyle.Thickness = 1.25f;
+
+	FDrawDebugPointStyle ThickPointStyle;
+	ThickPointStyle.Color = FColor::Black.WithAlpha(64);
+	ThickPointStyle.Thickness = 2.5f;
+
+	FDrawDebugLineStyle LineStyle;
+	LineStyle.Color = FColor::Black.WithAlpha(64);
+	LineStyle.Thickness = 1.0f;
+
 	TArray<TPair<FVector, FVector>> LineSegments;
 	LineSegments.Reserve(FrameNum);
 
@@ -2167,7 +2384,7 @@ void UDrawDebugLibrary::DrawDebugRangeTrajectoryArrayView(
 
 		if (bFirstOrLastPoint)
 		{
-			DrawDebugPoint(Drawer, CurrLocation, FColor::Black.WithAlpha(64), 2.5f, bDepthTest);
+			DrawDebugPoint(Drawer, CurrLocation, ThickPointStyle, bDepthTest);
 		}
 		else
 		{
@@ -2188,8 +2405,7 @@ void UDrawDebugLibrary::DrawDebugRangeTrajectoryArrayView(
 					Drawer,
 					CurrLocation,
 					CurrDirection.Rotation(),
-					FColor::Black.WithAlpha(64),
-					1.0f,
+					LineStyle,
 					bDepthTest,
 					15.0f, 15.0f);
 			}
@@ -2207,8 +2423,8 @@ void UDrawDebugLibrary::DrawDebugRangeTrajectoryArrayView(
 		}
 	}
 
-	DrawDebugPointsArrayView(Drawer, Points, FColor::Black.WithAlpha(64), 1.25f, bDepthTest);
-	DrawDebugLinesPairsArrayView(Drawer, LineSegments, FColor::Black.WithAlpha(64), 1.0f, bDepthTest);
+	DrawDebugPointsArrayView(Drawer, Points, PointStyle, bDepthTest);
+	DrawDebugLinesPairsArrayView(Drawer, LineSegments, LineStyle, bDepthTest);
 }
 
 int32 UDrawDebugLibrary::DrawDebugStringSegmentNum(const FString& String, const FDrawDebugStringSettings& Settings)
@@ -2292,12 +2508,12 @@ FVector UDrawDebugLibrary::DrawDebugNameDimensions(const FName& Name, const FDra
 }
 
 
-void UDrawDebugLibrary::DrawDebugString(const FDebugDrawer& Drawer, const FString& String, const FVector& Location, const FRotator& Rotation, const FLinearColor Color, const float Thickness, const bool bDepthTest, const FDrawDebugStringSettings& Settings)
+void UDrawDebugLibrary::DrawDebugString(const FDebugDrawer& Drawer, const FString& String, const FVector& Location, const FRotator& Rotation, const FDrawDebugLineStyle& LineStyle, const bool bDepthTest, const FDrawDebugStringSettings& Settings)
 {
-	DrawDebugStringView(Drawer, String, Location, Rotation, Color, Thickness, bDepthTest, Settings);
+	DrawDebugStringView(Drawer, String, Location, Rotation, LineStyle, bDepthTest, Settings);
 }
 
-void UDrawDebugLibrary::DrawDebugStringView(const FDebugDrawer& Drawer, const FStringView String, const FVector& Location, const FRotator& Rotation, const FLinearColor Color, const float Thickness, const bool bDepthTest, const FDrawDebugStringSettings& Settings)
+void UDrawDebugLibrary::DrawDebugStringView(const FDebugDrawer& Drawer, const FStringView String, const FVector& Location, const FRotator& Rotation, const FDrawDebugLineStyle& LineStyle, const bool bDepthTest, const FDrawDebugStringSettings& Settings)
 {
 	const int32 CharNum = String.Len();
 
@@ -2357,17 +2573,17 @@ void UDrawDebugLibrary::DrawDebugStringView(const FDebugDrawer& Drawer, const FS
 		}
 	}
 
-	DrawDebugLinesPairsArrayView(Drawer, LineSegments, Color, Thickness, bDepthTest);
+	DrawDebugLinesPairsArrayView(Drawer, LineSegments, LineStyle, bDepthTest);
 }
 
-void UDrawDebugLibrary::DrawDebugName(const FDebugDrawer& Drawer, const FName& Name, const FVector& Location, const FRotator& Rotation, const FLinearColor Color, const float Thickness, const bool bDepthTest, const FDrawDebugStringSettings& Settings)
+void UDrawDebugLibrary::DrawDebugName(const FDebugDrawer& Drawer, const FName& Name, const FVector& Location, const FRotator& Rotation, const FDrawDebugLineStyle& LineStyle, const bool bDepthTest, const FDrawDebugStringSettings& Settings)
 {
-	DrawDebugStringView(Drawer, Name.ToString(), Location, Rotation, Color, Thickness, bDepthTest, Settings);
+	DrawDebugStringView(Drawer, Name.ToString(), Location, Rotation, LineStyle, bDepthTest, Settings);
 }
 
-void UDrawDebugLibrary::DrawDebugText(const FDebugDrawer& Drawer, const FText& Text, const FVector& Location, const FRotator& Rotation, const FLinearColor Color, const float Thickness, const bool bDepthTest, const FDrawDebugStringSettings& Settings)
+void UDrawDebugLibrary::DrawDebugText(const FDebugDrawer& Drawer, const FText& Text, const FVector& Location, const FRotator& Rotation, const FDrawDebugLineStyle& LineStyle, const bool bDepthTest, const FDrawDebugStringSettings& Settings)
 {
-	DrawDebugStringView(Drawer, Text.ToString(), Location, Rotation, Color, Thickness, bDepthTest, Settings);
+	DrawDebugStringView(Drawer, Text.ToString(), Location, Rotation, LineStyle, bDepthTest, Settings);
 }
 
 
@@ -2392,14 +2608,13 @@ void UDrawDebugLibrary::VisualLoggerDrawText(const FDebugDrawer& Drawer, const F
 }
 
 
-void UDrawDebugLibrary::DrawDebugMoverOrientation(const FDebugDrawer& Drawer, const FVector& Location, const FRotator& Rotation, const FLinearColor Color, const float Thickness, const bool bDepthTest, const FVector ForwardVector, const float Scale)
+void UDrawDebugLibrary::DrawDebugMoverOrientation(const FDebugDrawer& Drawer, const FVector& Location, const FRotator& Rotation, const FDrawDebugLineStyle& LineStyle, const bool bDepthTest, const FVector ForwardVector, const float Scale)
 {
 	DrawDebugRotation(
 		Drawer,
 		Location,
 		Rotation,
-		Color.A,
-		Thickness,
+		LineStyle,
 		true,
 		Scale / 3.0f);
 
@@ -2411,8 +2626,7 @@ void UDrawDebugLibrary::DrawDebugMoverOrientation(const FDebugDrawer& Drawer, co
 		Drawer,
 		ArrowPosition,
 		ArrowDirection.Rotation(),
-		Color,
-		Thickness,
+		LineStyle,
 		true,
 		2.0f * ArrowSize,
 		2.0f * ArrowSize);
@@ -2421,8 +2635,7 @@ void UDrawDebugLibrary::DrawDebugMoverOrientation(const FDebugDrawer& Drawer, co
 		Drawer,
 		Location,
 		FRotator::ZeroRotator,
-		Color,
-		Thickness,
+		LineStyle,
 		true,
 		Scale,
 		31);
@@ -2440,10 +2653,9 @@ void UDrawDebugLibrary::DrawDebugGraph(
 	const float Ymax,
 	const float XaxisLength,
 	const float YaxisLength,
-	const FLinearColor Color,
-	const float AxisThickness,
-	const float LineThickness,
-	const float TextThickness,
+	const FDrawDebugLineStyle& TextLineStyle,
+	const FDrawDebugLineStyle& AxesLineStyle,
+	const FDrawDebugLineStyle& PlotLineStyle,
 	const bool bDepthTest,
 	const FDrawDebugGraphAxesSettings& AxesSettings)
 {
@@ -2459,10 +2671,9 @@ void UDrawDebugLibrary::DrawDebugGraph(
 		Ymax,
 		XaxisLength,
 		YaxisLength,
-		Color,
-		AxisThickness,
-		LineThickness,
-		TextThickness,
+		TextLineStyle,
+		AxesLineStyle,
+		PlotLineStyle,
 		bDepthTest,
 		AxesSettings);
 }
@@ -2479,10 +2690,9 @@ void UDrawDebugLibrary::DrawDebugGraphArrayView(
 	const float Ymax,
 	const float XaxisLength,
 	const float YaxisLength,
-	const FLinearColor Color,
-	const float AxisThickness,
-	const float LineThickness,
-	const float TextThickness,
+	const FDrawDebugLineStyle& TextLineStyle,
+	const FDrawDebugLineStyle& AxesLineStyle,
+	const FDrawDebugLineStyle& PlotLineStyle,
 	const bool bDepthTest,
 	const FDrawDebugGraphAxesSettings& AxesSettings)
 {
@@ -2492,9 +2702,8 @@ void UDrawDebugLibrary::DrawDebugGraphArrayView(
 		Rotation,
 		XaxisLength,
 		YaxisLength,
-		Color,
-		AxisThickness,
-		TextThickness,
+		TextLineStyle,
+		AxesLineStyle,
 		bDepthTest,
 		AxesSettings);
 
@@ -2510,8 +2719,7 @@ void UDrawDebugLibrary::DrawDebugGraphArrayView(
 		Ymax,
 		XaxisLength,
 		YaxisLength,
-		Color,
-		LineThickness,
+		PlotLineStyle,
 		bDepthTest);
 }
 
@@ -2521,9 +2729,8 @@ void UDrawDebugLibrary::DrawDebugGraphAxes(
 	const FRotator& Rotation,
 	const float XaxisLength,
 	const float YaxisLength,
-	const FLinearColor Color,
-	const float AxisThickness,
-	const float TextThickness,
+	const FDrawDebugLineStyle& TextLineStyle,
+	const FDrawDebugLineStyle& AxesLineStyle,
 	const bool bDepthTest,
 	const FDrawDebugGraphAxesSettings& AxesSettings)
 {
@@ -2532,8 +2739,7 @@ void UDrawDebugLibrary::DrawDebugGraphAxes(
 		Location,
 		(Rotation.Quaternion() * FQuat::MakeFromRotationVector(FVector(UE_PI / 2, 0, 0))).Rotator(),
 		XaxisLength,
-		Color,
-		AxisThickness,
+		AxesLineStyle,
 		bDepthTest);
 
 	DrawDebugOrientedArrow(
@@ -2543,8 +2749,7 @@ void UDrawDebugLibrary::DrawDebugGraphAxes(
 			FQuat::MakeFromRotationVector(FVector(0, -UE_PI / 2, 0)) *
 			FQuat::MakeFromRotationVector(FVector(UE_PI / 2, 0, 0))).Rotator(),
 		YaxisLength,
-		Color,
-		AxisThickness,
+		AxesLineStyle,
 		bDepthTest);
 
 	const FVector TitleDimension = DrawDebugStringViewDimensions(AxesSettings.Title, AxesSettings.TitleSettings);
@@ -2554,8 +2759,7 @@ void UDrawDebugLibrary::DrawDebugGraphAxes(
 		AxesSettings.Title,
 		Location + Rotation.RotateVector(FVector(XaxisLength / 2.0f - TitleDimension.X / 2.0f, 0.0, YaxisLength + 1.1f * TitleDimension.Y)),
 		Rotation,
-		Color,
-		TextThickness,
+		TextLineStyle,
 		bDepthTest,
 		AxesSettings.TitleSettings);
 
@@ -2566,8 +2770,7 @@ void UDrawDebugLibrary::DrawDebugGraphAxes(
 		AxesSettings.XaxisLabel,
 		Location + Rotation.RotateVector(FVector(XaxisLength / 2.0f - XlabelDimension.X / 2.0f, 0.0, -0.1f * XlabelDimension.Y)),
 		Rotation,
-		Color,
-		TextThickness,
+		TextLineStyle,
 		bDepthTest,
 		AxesSettings.AxisLabelSettings);
 
@@ -2578,8 +2781,7 @@ void UDrawDebugLibrary::DrawDebugGraphAxes(
 		AxesSettings.YaxisLabel,
 		Location + Rotation.RotateVector(FVector(-YlabelDimension.Y - 0.1f * YlabelDimension.Y, 0.0, YaxisLength / 2 - YlabelDimension.X / 2)),
 		(Rotation.Quaternion() * FQuat::MakeFromRotationVector(FVector(0, -UE_PI / 2, 0))).Rotator(),
-		Color,
-		TextThickness,
+		TextLineStyle,
 		bDepthTest,
 		AxesSettings.AxisLabelSettings);
 }
@@ -2596,11 +2798,10 @@ void UDrawDebugLibrary::DrawDebugGraphLine(
 	const float Ymax,
 	const float XaxisLength,
 	const float YaxisLength,
-	const FLinearColor Color,
-	const float LineThickness,
+	const FDrawDebugLineStyle& LineStyle,
 	const bool bDepthTest)
 {
-	DrawDebugGraphLineArrayView(Drawer, Location, Rotation, Xvalues, Yvalues, Xmin, Xmax, Ymin, Ymax, XaxisLength, YaxisLength, Color, LineThickness, bDepthTest);
+	DrawDebugGraphLineArrayView(Drawer, Location, Rotation, Xvalues, Yvalues, Xmin, Xmax, Ymin, Ymax, XaxisLength, YaxisLength, LineStyle, bDepthTest);
 }
 
 void UDrawDebugLibrary::DrawDebugGraphLineArrayView(
@@ -2615,8 +2816,7 @@ void UDrawDebugLibrary::DrawDebugGraphLineArrayView(
 	const float Ymax,
 	const float XaxisLength,
 	const float YaxisLength,
-	const FLinearColor Color,
-	const float LineThickness,
+	const FDrawDebugLineStyle& LineStyle,
 	const bool bDepthTest)
 {
 	const int32 PointNum = FMath::Min(Xvalues.Num(), Yvalues.Num());
@@ -2639,7 +2839,7 @@ void UDrawDebugLibrary::DrawDebugGraphLineArrayView(
 			Location + Rotation.RotateVector(FVector(Xloc1, 0.0, Yloc1)) });
 	}
 
-	DrawDebugLinesPairsArrayView(Drawer, LineSegments, Color, LineThickness, bDepthTest);
+	DrawDebugLinesPairsArrayView(Drawer, LineSegments, LineStyle, bDepthTest);
 }
 
 void UDrawDebugLibrary::DrawDebugGraphLegend(
@@ -2650,12 +2850,13 @@ void UDrawDebugLibrary::DrawDebugGraphLegend(
 	const TArray<FString>& LegendLabels,
 	const float XaxisLength,
 	const float YaxisLength,
-	const float TextThickness,
-	const float PointSize,
+	const FDrawDebugLineStyle& TextLineStyle,
+	const FDrawDebugLineStyle& IconLineStyle,
+	const float IconSize,
 	const bool bDepthTest,
 	const FDrawDebugStringSettings& LegendSettings)
 {
-	DrawDebugGraphLegendArrayView(Drawer, Location, Rotation, LegendColors, LegendLabels, XaxisLength, YaxisLength, TextThickness, PointSize, bDepthTest, LegendSettings);
+	DrawDebugGraphLegendArrayView(Drawer, Location, Rotation, LegendColors, LegendLabels, XaxisLength, YaxisLength, TextLineStyle, IconLineStyle, IconSize, bDepthTest, LegendSettings);
 }
 
 void UDrawDebugLibrary::DrawDebugGraphLegendArrayView(
@@ -2666,8 +2867,9 @@ void UDrawDebugLibrary::DrawDebugGraphLegendArrayView(
 	const TArrayView<const FString> LegendLabels,
 	const float XaxisLength,
 	const float YaxisLength,
-	const float TextThickness,
-	const float PointSize,
+	const FDrawDebugLineStyle& TextLineStyle,
+	const FDrawDebugLineStyle& IconLineStyle,
+	const float IconSize,
 	const bool bDepthTest,
 	const FDrawDebugStringSettings& LegendSettings)
 {
@@ -2688,17 +2890,16 @@ void UDrawDebugLibrary::DrawDebugGraphLegendArrayView(
 
 		DrawDebugDiamond(
 			Drawer, 
-			Location + Rotation.RotateVector(FVector(XaxisLength + 2.0f * PointSize, 0.0, VerticalOffset - LegendHeight / 2.0f)),
+			Location + Rotation.RotateVector(FVector(XaxisLength + 2.0f * IconSize, 0.0, VerticalOffset - LegendHeight / 2.0f)),
 			(Rotation.Quaternion() * IconRotation).Rotator(),
-			LegendColors[LegendIdx], TextThickness, bDepthTest, PointSize / 2.0f);
+			DrawDebugLineStyleWithColor(IconLineStyle, LegendColors[LegendIdx]), bDepthTest, IconSize / 2.0f);
 
 		DrawDebugStringView(
 			Drawer,
 			LegendLabels[LegendIdx],
-			Location + Rotation.RotateVector(FVector(XaxisLength + 4.0f * PointSize, 0.0, VerticalOffset)),
+			Location + Rotation.RotateVector(FVector(XaxisLength + 4.0f * IconSize, 0.0, VerticalOffset)),
 			Rotation,
-			LegendColors[LegendIdx],
-			TextThickness,
+			DrawDebugLineStyleWithColor(TextLineStyle, LegendColors[LegendIdx]),
 			bDepthTest,
 			LegendSettings);
 
diff --git a/Engine/Plugins/Experimental/DrawDebugLibrary/Source/DrawDebugLibrary/Public/DrawDebugLibrary.h b/Engine/Plugins/Experimental/DrawDebugLibrary/Source/DrawDebugLibrary/Public/DrawDebugLibrary.h
--- a/Engine/Plugins/Experimental/DrawDebugLibrary/Source/DrawDebugLibrary/Public/DrawDebugLibrary.h
+++ b/Engine/Plugins/Experimental/DrawDebugLibrary/Source/DrawDebugLibrary/Public/DrawDebugLibrary.h
@@ -62,6 +62,70 @@ enum class EDrawDebugLogVerbosity : uint8
 	VeryVerbose,
 };
 
+
+/** Style for debug drawing points */
+USTRUCT(BlueprintType, Category = "Draw Debug Library")
+struct FDrawDebugPointStyle
+{
+	GENERATED_BODY()
+
+	/** Thickness of the point. Set to 0.0 for pixel-sized points. */
+	UPROPERTY(BlueprintReadWrite, EditAnywhere, Category = "Draw Debug Library", meta = (ClampMin = "0", UIMin = "0", Units = "cm"))
+	float Thickness = 0.0f;
+
+	/** Color of the line. Alpha controls the opacity. */
+	UPROPERTY(BlueprintReadWrite, EditAnywhere, Category = "Draw Debug Library")
+	FLinearColor Color = FLinearColor(1.0f, 0.0f, 1.0f, 1.0f);
+};
+
+
+/** Enum for different line types */
+UENUM(BlueprintType, Category = "Draw Debug Library")
+enum class EDrawDebugLineType : uint8
+{
+	/** Normal solid line */
+	Solid,
+
+	/** Line made up of multiple dashes */
+	Dashed,
+
+	/** Line drawn using points */
+	Dotted,
+};
+
+
+/** Style for debug drawing lines */
+USTRUCT(BlueprintType, Category = "Draw Debug Library")
+struct FDrawDebugLineStyle
+{
+	GENERATED_BODY()
+
+	/** Line style to draw with */
+	UPROPERTY(BlueprintReadWrite, EditAnywhere, Category = "Draw Debug Library")
+	EDrawDebugLineType LineType = EDrawDebugLineType::Solid;
+
+	/** Thickness of the line. Set to 0.0 for pixel-sized lines. */
+	UPROPERTY(BlueprintReadWrite, EditAnywhere, Category = "Draw Debug Library", meta = (ClampMin = "0", UIMin = "0", Units = "cm"))
+	float Thickness = 0.0f;
+
+	/** Color of the line. Alpha controls the opacity. */
+	UPROPERTY(BlueprintReadWrite, EditAnywhere, Category = "Draw Debug Library")
+	FLinearColor Color = FLinearColor(1.0f, 0.0f, 1.0f, 1.0f);
+
+	/** Width of the dashes when using dashed style. */
+	UPROPERTY(BlueprintReadWrite, EditAnywhere, Category = "Draw Debug Library", meta = (ClampMin = "0", UIMin = "0", Units = "cm"))
+	float DashWidth = 2.5f;
+
+	/** Spacing of the dashes when using dashed style. */
+	UPROPERTY(BlueprintReadWrite, EditAnywhere, Category = "Draw Debug Library", meta = (ClampMin = "0", UIMin = "0", Units = "cm"))
+	float DashSpacing = 1.0f;
+
+	/** Spacing of the dots when using dotted style. */
+	UPROPERTY(BlueprintReadWrite, EditAnywhere, Category = "Draw Debug Library", meta = (ClampMin = "0", UIMin = "0", Units = "cm"))
+	float DotSpacing = 1.0f;
+};
+
+
 /** Enum for different arrow head types */
 UENUM(BlueprintType, Category = "Draw Debug Library")
 enum class EDrawDebugArrowHead : uint8
@@ -73,13 +137,31 @@ enum class EDrawDebugArrowHead : uint8
 	FourWay,
 
 	/** Triangular-based pyramid arrow head */
-	Triangular,
+	TriangularPyramid,
 
 	/** Square-based pyramid arrow head */
-	Square,
+	SquarePyramid,
 
 	/** Cone arrow head */
 	Cone,
+
+	/** Flat triangle arrow head */
+	Triangle,
+
+	/** Flat square arrow head */
+	Square,
+
+	/** Flat diamond arrow head */
+	Diamond,
+
+	/** Flat circle arrow head */
+	Circle,
+
+	/** 3D sphere arrow head */
+	Sphere,
+
+	/** 3D box arrow head */
+	Box,
 };
 
 
@@ -93,9 +175,13 @@ struct FDrawDebugArrowSettings
 	UPROPERTY(BlueprintReadWrite, EditAnywhere, Category = "Draw Debug Library")
 	bool bArrowheadOnStart = false;
 
-	/** Width of the start arrow head */
+	/** If the line should end at the start arrow head */
+	UPROPERTY(BlueprintReadWrite, EditAnywhere, Category = "Draw Debug Library", meta = (EditCondition = "bArrowheadOnStart", HideEditConditionToggle))
+	bool bArrowLineEndsAtStartHead = false;
+
+	/** Size of the start arrow head */
 	UPROPERTY(BlueprintReadWrite, EditAnywhere, Category = "Draw Debug Library", meta = (EditCondition = "bArrowheadOnStart", HideEditConditionToggle, ClampMin = "0", UIMin = "0", Units = "cm"))
-	float ArrowHeadStartWidth = 5.0f;
+	float ArrowHeadStartSize = 5.0f;
 
 	/** Type of start arrow head */
 	UPROPERTY(BlueprintReadWrite, EditAnywhere, Category = "Draw Debug Library", meta = (EditCondition = "bArrowheadOnStart", HideEditConditionToggle))
@@ -105,9 +191,13 @@ struct FDrawDebugArrowSettings
 	UPROPERTY(BlueprintReadWrite, EditAnywhere, Category = "Draw Debug Library")
 	bool bArrowheadOnEnd = true;
 
-	/** Width of the end arrow head */
+	/** If the line should end at the end arrow head */
+	UPROPERTY(BlueprintReadWrite, EditAnywhere, Category = "Draw Debug Library", meta = (EditCondition = "bArrowheadOnEnd", HideEditConditionToggle))
+	bool bArrowLineEndsAtEndHead = false;
+
+	/** Size of the end arrow head */
 	UPROPERTY(BlueprintReadWrite, EditAnywhere, Category = "Draw Debug Library", meta = (EditCondition = "bArrowheadOnEnd", HideEditConditionToggle, ClampMin = "0", UIMin = "0", Units = "cm"))
-	float ArrowHeadEndWidth = 5.0f;
+	float ArrowHeadEndSize = 5.0f;
 
 	/** Type of end arrow head */
 	UPROPERTY(BlueprintReadWrite, EditAnywhere, Category = "Draw Debug Library", meta = (EditCondition = "bArrowheadOnEnd", HideEditConditionToggle))
@@ -286,7 +376,7 @@ struct FDrawDebugGraphAxesSettings
  * functions for debug drawing if we are drawing using a UWorld object or a FPrimitiveDrawInterface object (for example). Since this contains raw 
  * pointers it should not be stored anywhere as these could become stale. It should always be treated as a temporary object.
  */
-USTRUCT(BlueprintType, Category = "Draw Debug Library")
+USTRUCT(BlueprintType, Category = "Draw Debug Library", meta = (HasNativeMake = "/Script/DrawDebugLibrary.DrawDebugLibrary.MakeDebugDrawer"))
 struct FDebugDrawer
 {
 	GENERATED_BODY()
@@ -345,8 +435,7 @@ public: // Constructors
 public: // Drawer specific interface
 
 	UE_API void DrawPoints(const TArrayView<const FVector> Locations, const FLinearColor& Color, const float Thickness, const bool bDepthTest) const;
-	UE_API void DrawLines(const TArrayView<const FVector> StartLocations, const TArrayView<const FVector> EndLocations, const FLinearColor& Color, const float Thickness, const bool bDepthTest) const;
-	UE_API void DrawLinesPairs(const TArrayView<const TPair<FVector, FVector>> Segments, const FLinearColor& Color, const float Thickness, const bool bDepthTest) const;
+	UE_API void DrawLines(const TArrayView<const TPair<FVector, FVector>> Segments, const FLinearColor& Color, const float Thickness, const bool bDepthTest) const;
 
 	UE_API void VisualLoggerLogString(const FStringView String, const FLinearColor& Color) const;
 	UE_API void VisualLoggerDrawString(const FStringView String, const FVector& Location, const FLinearColor& Color) const;
@@ -387,7 +476,7 @@ private:
  * 
  * General parameter structure for these functions is:
  * 
- *	Required Inputs, Color, Thickness, bDepthTest, Settings
+ *	Required Inputs, Style, bDepthTest, Settings
  */
 UCLASS(BlueprintType, Category = "Draw Debug Library")
 class UDrawDebugLibrary : public UBlueprintFunctionLibrary
@@ -451,12 +540,12 @@ public: // Visual Debugger Logging
 
 	/** Log a string to the visual logger. Only does anything with a VisualLoggerDebugDrawer. */
 	UFUNCTION(BlueprintCallable, Category = "Draw Debug Library", meta = (DevelopmentOnly, AutoCreateRefTerm = "String"))
-	static UE_API void VisualLoggerLogString(const FDebugDrawer& Drawer, const FString& String, const FLinearColor Color = FLinearColor::Red);
-	static UE_API void VisualLoggerLogStringView(const FDebugDrawer& Drawer, const FStringView String, const FLinearColor Color = FLinearColor::Red);
+	static UE_API void VisualLoggerLogString(const FDebugDrawer& Drawer, const FString& String, const FLinearColor Color = FLinearColor(1.0f, 0.0f, 1.0f, 1.0f));
+	static UE_API void VisualLoggerLogStringView(const FDebugDrawer& Drawer, const FStringView String, const FLinearColor Color = FLinearColor(1.0f, 0.0f, 1.0f, 1.0f));
 
 	/** Log some text to the visual logger. Only does anything with a VisualLoggerDebugDrawer. */
 	UFUNCTION(BlueprintCallable, Category = "Draw Debug Library", meta = (DevelopmentOnly, AutoCreateRefTerm = "Text"))
-	static UE_API void VisualLoggerLogText(const FDebugDrawer& Drawer, const FText& Text, const FLinearColor Color = FLinearColor::Red);
+	static UE_API void VisualLoggerLogText(const FDebugDrawer& Drawer, const FText& Text, const FLinearColor Color = FLinearColor(1.0f, 0.0f, 1.0f, 1.0f));
 	
 public: // Convenience functions
 
@@ -472,220 +561,263 @@ public: // Convenience functions
 	UFUNCTION(BlueprintPure = false, Category = "Draw Debug Library")
 	static UE_API void AddToVectorHistoryArray(UPARAM(ref) TArray<FVector>& InOutValues, const FVector NewValue = FVector::ZeroVector, const int32 MaxHistoryNum = 60);
 
-public: // Point & Line
+public: // Style Modifiers
+
+	/** Make a point style from a color */
+	UFUNCTION(BlueprintPure, Category = "Draw Debug Library", meta = (BlueprintAutocast, CompactNodeTitle = "->"))
+	static UE_API FDrawDebugPointStyle MakeDrawDebugPointStyleFromColor(const FLinearColor Color = FLinearColor(1.0f, 0.0f, 1.0f, 1.0f));
+
+	/** Make a line style from a color */
+	UFUNCTION(BlueprintPure, Category = "Draw Debug Library", meta = (BlueprintAutocast, CompactNodeTitle = "->"))
+	static UE_API FDrawDebugLineStyle MakeDrawDebugLineStyleFromColor(const FLinearColor Color = FLinearColor(1.0f, 0.0f, 1.0f, 1.0f));
+
+	/** Get a point style from a line style */
+	UFUNCTION(BlueprintPure, Category = "Draw Debug Library", meta = (BlueprintAutocast, CompactNodeTitle = "->", AutoCreateRefTerm = "LineStyle"))
+	static UE_API FDrawDebugPointStyle DrawDebugPointStyleFromLineStyle(const FDrawDebugLineStyle& LineStyle);
+
+	/** Get a line style from a point style */
+	UFUNCTION(BlueprintPure, Category = "Draw Debug Library", meta = (BlueprintAutocast, CompactNodeTitle = "->", AutoCreateRefTerm = "PointStyle"))
+	static UE_API FDrawDebugLineStyle DrawDebugLineStyleFromPointStyle(const FDrawDebugPointStyle& PointStyle);
+
+	/** Get the color from a point style */
+	UFUNCTION(BlueprintPure, Category = "Draw Debug Library", meta = (BlueprintAutocast, CompactNodeTitle = "->", AutoCreateRefTerm = "PointStyle"))
+	static UE_API FLinearColor DrawDebugPointStyleColor(const FDrawDebugPointStyle& PointStyle);
+
+	/** Get the color from a line style */
+	UFUNCTION(BlueprintPure, Category = "Draw Debug Library", meta = (BlueprintAutocast, CompactNodeTitle = "->", AutoCreateRefTerm = "LineStyle"))
+	static UE_API FLinearColor DrawDebugLineStyleColor(const FDrawDebugLineStyle& LineStyle);
+
+	/** Get a line style with the color changed */
+	UFUNCTION(BlueprintPure, Category = "Draw Debug Library", meta = (AutoCreateRefTerm = "LineStyle"))
+	static UE_API FDrawDebugLineStyle DrawDebugLineStyleWithColor(const FDrawDebugLineStyle& LineStyle, const FLinearColor Color = FLinearColor(1.0f, 0.0f, 1.0f, 1.0f));
+
+	/** Get a line style with the color (excluding the opacity) changed */
+	UFUNCTION(BlueprintPure, Category = "Draw Debug Library", meta = (AutoCreateRefTerm = "LineStyle"))
+	static UE_API FDrawDebugLineStyle DrawDebugLineStyleWithColorNoOpacity(const FDrawDebugLineStyle& LineStyle, const FLinearColor Color = FLinearColor(1.0f, 0.0f, 1.0f, 1.0f));
+
+	/** Get a line style with the thickness changed */
+	UFUNCTION(BlueprintPure, Category = "Draw Debug Library", meta = (AutoCreateRefTerm = "LineStyle"))
+	static UE_API FDrawDebugLineStyle DrawDebugLineStyleWithThickness(const FDrawDebugLineStyle& LineStyle, const float Thickness = 0.0f);
+
+	/** Get a line style with the type changed */
+	UFUNCTION(BlueprintPure, Category = "Draw Debug Library", meta = (AutoCreateRefTerm = "LineStyle"))
+	static UE_API FDrawDebugLineStyle DrawDebugLineStyleWithType(const FDrawDebugLineStyle& LineStyle, const EDrawDebugLineType LineType = EDrawDebugLineType::Solid);
+
+public: // Points
 
 	/** Debug Draw a point */
-	UFUNCTION(BlueprintCallable, Category = "Draw Debug Library", meta = (DevelopmentOnly, AutoCreateRefTerm = "Location"))
-	static UE_API void DrawDebugPoint(const FDebugDrawer& Drawer, const FVector& Location, const FLinearColor Color = FLinearColor::Red, const float Thickness = 0.0f, const bool bDepthTest = true);
+	UFUNCTION(BlueprintCallable, Category = "Draw Debug Library", meta = (DevelopmentOnly, AutoCreateRefTerm = "Location, PointStyle"))
+	static UE_API void DrawDebugPoint(const FDebugDrawer& Drawer, const FVector& Location, const FDrawDebugPointStyle& PointStyle = FDrawDebugPointStyle(), const bool bDepthTest = true);
 
 	/** Debug Draw points. These should be preferred over DrawDebugPoint where possible as it will batch drawing when required such as when using the visual logger. */
-	UFUNCTION(BlueprintCallable, Category = "Draw Debug Library", meta = (DevelopmentOnly))
-	static UE_API void DrawDebugPoints(const FDebugDrawer& Drawer, const TArray<FVector>& Locations, const FLinearColor Color = FLinearColor::Red, const float Thickness = 0.0f, const bool bDepthTest = true);
-	static UE_API void DrawDebugPointsArrayView(const FDebugDrawer& Drawer, const TArrayView<const FVector> Locations, const FLinearColor Color = FLinearColor::Red, const float Thickness = 0.0f, const bool bDepthTest = true);
+	UFUNCTION(BlueprintCallable, Category = "Draw Debug Library", meta = (DevelopmentOnly, AutoCreateRefTerm = "Locations, PointStyle"))
+	static UE_API void DrawDebugPoints(const FDebugDrawer& Drawer, const TArray<FVector>& Locations, const FDrawDebugPointStyle& PointStyle = FDrawDebugPointStyle(), const bool bDepthTest = true);
+	static UE_API void DrawDebugPointsArrayView(const FDebugDrawer& Drawer, const TArrayView<const FVector> Locations, const FDrawDebugPointStyle& PointStyle = FDrawDebugPointStyle(), const bool bDepthTest = true);
+
+public: // Lines
 
 	/** Debug Draw a line */
-	UFUNCTION(BlueprintCallable, Category = "Draw Debug Library", meta = (DevelopmentOnly, AutoCreateRefTerm = "StartLocation, EndLocation"))
-	static UE_API void DrawDebugLine(const FDebugDrawer& Drawer, const FVector& StartLocation, const FVector& EndLocation, const FLinearColor Color = FLinearColor::Red, const float Thickness = 0.0f, const bool bDepthTest = true);
+	UFUNCTION(BlueprintCallable, Category = "Draw Debug Library", meta = (DevelopmentOnly, AutoCreateRefTerm = "StartLocation, EndLocation, LineStyle"))
+	static UE_API void DrawDebugLine(const FDebugDrawer& Drawer, const FVector& StartLocation, const FVector& EndLocation, const FDrawDebugLineStyle& LineStyle = FDrawDebugLineStyle(), const bool bDepthTest = true);
 
 	/** Debug Draw lines. These should be preferred over DrawDebugLine where possible as it will batch drawing when required such as when using the visual logger. */
-	UFUNCTION(BlueprintCallable, Category = "Draw Debug Library", meta = (DevelopmentOnly))
-	static UE_API void DrawDebugLines(const FDebugDrawer& Drawer, const TArray<FVector>& StartLocations, const TArray<FVector>& EndLocations, const FLinearColor Color = FLinearColor::Red, const float Thickness = 0.0f, const bool bDepthTest = true);
-	static UE_API void DrawDebugLinesArrayView(const FDebugDrawer& Drawer, const TArrayView<const FVector> StartLocations, const TArrayView<const FVector> EndLocations, const FLinearColor Color = FLinearColor::Red, const float Thickness = 0.0f, const bool bDepthTest = true);
-	static UE_API void DrawDebugLinesPairsArrayView(const FDebugDrawer& Drawer, const TArrayView<const TPair<FVector, FVector>> Segments, const FLinearColor Color = FLinearColor::Red, const float Thickness = 0.0f, const bool bDepthTest = true);
+	UFUNCTION(BlueprintCallable, Category = "Draw Debug Library", meta = (DevelopmentOnly, AutoCreateRefTerm = "StartLocations, EndLocations, LineStyle"))
+	static UE_API void DrawDebugLines(const FDebugDrawer& Drawer, const TArray<FVector>& StartLocations, const TArray<FVector>& EndLocations, const FDrawDebugLineStyle& LineStyle = FDrawDebugLineStyle(), const bool bDepthTest = true);
+	static UE_API void DrawDebugLinesArrayView(const FDebugDrawer& Drawer, const TArrayView<const FVector> StartLocations, const TArrayView<const FVector> EndLocations, const FDrawDebugLineStyle& LineStyle = FDrawDebugLineStyle(), const bool bDepthTest = true);
+	static UE_API void DrawDebugLinesPairsArrayView(const FDebugDrawer& Drawer, const TArrayView<const TPair<FVector, FVector>> Segments, const FDrawDebugLineStyle& LineStyle = FDrawDebugLineStyle(), const bool bDepthTest = true);
 
 public: // Basic Shapes
 
 	/** Debug Draw a triangular base pyramid */
-	UFUNCTION(BlueprintCallable, Category = "Draw Debug Library", meta = (DevelopmentOnly, AutoCreateRefTerm = "Location, Rotation"))
-	static UE_API void DrawDebugTriangularBasePyramid(const FDebugDrawer& Drawer, const FVector& Location, const FRotator& Rotation, const FLinearColor Color = FLinearColor::Red, const float Thickness = 0.0f, const bool bDepthTest = true, const float Length = 20.0f, const float Width = 20.0f);
+	UFUNCTION(BlueprintCallable, Category = "Draw Debug Library", meta = (DevelopmentOnly, AutoCreateRefTerm = "Location, Rotation, LineStyle"))
+	static UE_API void DrawDebugTriangularBasePyramid(const FDebugDrawer& Drawer, const FVector& Location, const FRotator& Rotation, const FDrawDebugLineStyle& LineStyle = FDrawDebugLineStyle(), const bool bDepthTest = true, const float Length = 20.0f, const float Width = 20.0f);
 
 	/** Debug Draw a square base pyramid */
-	UFUNCTION(BlueprintCallable, Category = "Draw Debug Library", meta = (DevelopmentOnly, AutoCreateRefTerm = "Location, Rotation"))
-	static UE_API void DrawDebugSquareBasePyramid(const FDebugDrawer& Drawer, const FVector& Location, const FRotator& Rotation, const FLinearColor Color = FLinearColor::Red, const float Thickness = 0.0f, const bool bDepthTest = true, const float Length = 20.0f, const float Width = 20.0f);
+	UFUNCTION(BlueprintCallable, Category = "Draw Debug Library", meta = (DevelopmentOnly, AutoCreateRefTerm = "Location, Rotation, LineStyle"))
+	static UE_API void DrawDebugSquareBasePyramid(const FDebugDrawer& Drawer, const FVector& Location, const FRotator& Rotation, const FDrawDebugLineStyle& LineStyle = FDrawDebugLineStyle(), const bool bDepthTest = true, const float Length = 20.0f, const float Width = 20.0f);
 
 	/** Debug Draw a cone */
-	UFUNCTION(BlueprintCallable, Category = "Draw Debug Library", meta = (DevelopmentOnly, AutoCreateRefTerm = "Location, Rotation"))
-	static UE_API void DrawDebugCone(const FDebugDrawer& Drawer, const FVector& Location, const FRotator& Rotation, const FLinearColor Color = FLinearColor::Red, const float Thickness = 0.0f, const bool bDepthTest = true, const float Length = 20.0f, const float Radius = 10.0f, const int32 Segments = 9);
+	UFUNCTION(BlueprintCallable, Category = "Draw Debug Library", meta = (DevelopmentOnly, AutoCreateRefTerm = "Location, Rotation, LineStyle"))
+	static UE_API void DrawDebugCone(const FDebugDrawer& Drawer, const FVector& Location, const FRotator& Rotation, const FDrawDebugLineStyle& LineStyle = FDrawDebugLineStyle(), const bool bDepthTest = true, const float Length = 20.0f, const float Radius = 10.0f, const int32 Segments = 9);
 
 	/** Debug Draw an arc around the XY axis */
-	UFUNCTION(BlueprintCallable, Category = "Draw Debug Library", meta = (DevelopmentOnly, AutoCreateRefTerm = "Location, Rotation"))
-	static UE_API void DrawDebugArc(const FDebugDrawer& Drawer, const FVector& Location, const FRotator& Rotation, const float Angle = 360.0f, const FLinearColor Color = FLinearColor::Red, const float Thickness = 0.0f, const bool bDepthTest = true, const float Radius = 10.0f, const int32 Segments = 17);
+	UFUNCTION(BlueprintCallable, Category = "Draw Debug Library", meta = (DevelopmentOnly, AutoCreateRefTerm = "Location, Rotation, LineStyle"))
+	static UE_API void DrawDebugArc(const FDebugDrawer& Drawer, const FVector& Location, const FRotator& Rotation, const float Angle = 360.0f, const FDrawDebugLineStyle& LineStyle = FDrawDebugLineStyle(), const bool bDepthTest = true, const float Radius = 10.0f, const int32 Segments = 17);
 
 	/** Debug Draw an circle around the XY axis */
-	UFUNCTION(BlueprintCallable, Category = "Draw Debug Library", meta = (DevelopmentOnly, AutoCreateRefTerm = "Location, Rotation"))
-	static UE_API void DrawDebugCircle(const FDebugDrawer& Drawer, const FVector& Location, const FRotator& Rotation, const FLinearColor Color = FLinearColor::Red, const float Thickness = 0.0f, const bool bDepthTest = true, const float Radius = 10.0f, const int32 Segments = 17);
+	UFUNCTION(BlueprintCallable, Category = "Draw Debug Library", meta = (DevelopmentOnly, AutoCreateRefTerm = "Location, Rotation, LineStyle"))
+	static UE_API void DrawDebugCircle(const FDebugDrawer& Drawer, const FVector& Location, const FRotator& Rotation, const FDrawDebugLineStyle& LineStyle = FDrawDebugLineStyle(), const bool bDepthTest = true, const float Radius = 10.0f, const int32 Segments = 17);
 
 	/** Debug Draw an tick on a circle at a given angle */
-	UFUNCTION(BlueprintCallable, Category = "Draw Debug Library", meta = (DevelopmentOnly, AutoCreateRefTerm = "Location, Rotation"))
-	static UE_API void DrawDebugCircleTick(const FDebugDrawer& Drawer, const FVector& Location, const FRotator& Rotation, const float Angle, const FLinearColor Color = FLinearColor::Red, const float Thickness = 0.0f, const bool bDepthTest = true, const float Radius = 10.0f, const float Length = 2.5f, bool bInside = true);
+	UFUNCTION(BlueprintCallable, Category = "Draw Debug Library", meta = (DevelopmentOnly, AutoCreateRefTerm = "Location, Rotation, LineStyle"))
+	static UE_API void DrawDebugCircleTick(const FDebugDrawer& Drawer, const FVector& Location, const FRotator& Rotation, const float Angle, const FDrawDebugLineStyle& LineStyle = FDrawDebugLineStyle(), const bool bDepthTest = true, const float Radius = 10.0f, const float Length = 2.5f, bool bInside = true);
 
 	/** Debug Draw ticks on a circle at the given angles */
-	UFUNCTION(BlueprintCallable, Category = "Draw Debug Library", meta = (DevelopmentOnly, AutoCreateRefTerm = "Location, Rotation, Angles"))
-	static UE_API void DrawDebugCircleTicks(const FDebugDrawer& Drawer, const FVector& Location, const FRotator& Rotation, const TArray<float>& Angles, const FLinearColor Color = FLinearColor::Red, const float Thickness = 0.0f, const bool bDepthTest = true, const float Radius = 10.0f, const float Length = 5.0f, bool bInside = true);
-	static UE_API void DrawDebugCircleTicksArrayView(const FDebugDrawer& Drawer, const FVector& Location, const FRotator& Rotation, const TArrayView<const float> Angles, const FLinearColor Color = FLinearColor::Red, const float Thickness = 0.0f, const bool bDepthTest = true, const float Radius = 10.0f, const float Length = 5.0f, bool bInside = true);
+	UFUNCTION(BlueprintCallable, Category = "Draw Debug Library", meta = (DevelopmentOnly, AutoCreateRefTerm = "Location, Rotation, Angles, LineStyle"))
+	static UE_API void DrawDebugCircleTicks(const FDebugDrawer& Drawer, const FVector& Location, const FRotator& Rotation, const TArray<float>& Angles, const FDrawDebugLineStyle& LineStyle = FDrawDebugLineStyle(), const bool bDepthTest = true, const float Radius = 10.0f, const float Length = 5.0f, bool bInside = true);
+	static UE_API void DrawDebugCircleTicksArrayView(const FDebugDrawer& Drawer, const FVector& Location, const FRotator& Rotation, const TArrayView<const float> Angles, const FDrawDebugLineStyle& LineStyle = FDrawDebugLineStyle(), const bool bDepthTest = true, const float Radius = 10.0f, const float Length = 5.0f, bool bInside = true);
 
 	/** Debug Draw the outline of a circle (two circles) around the XY axis */
-	UFUNCTION(BlueprintCallable, Category = "Draw Debug Library", meta = (DevelopmentOnly, AutoCreateRefTerm = "Location, Rotation"))
-	static UE_API void DrawDebugCircleOutline(const FDebugDrawer& Drawer, const FVector& Location, const FRotator& Rotation, const FLinearColor Color = FLinearColor::Red, const float Thickness = 0.0f, const bool bDepthTest = true, const float InnerRadius = 10.0f, const float OuterRadius = 15.0f, const int32 Segments = 17);
+	UFUNCTION(BlueprintCallable, Category = "Draw Debug Library", meta = (DevelopmentOnly, AutoCreateRefTerm = "Location, Rotation, LineStyle"))
+	static UE_API void DrawDebugCircleOutline(const FDebugDrawer& Drawer, const FVector& Location, const FRotator& Rotation, const FDrawDebugLineStyle& LineStyle = FDrawDebugLineStyle(), const bool bDepthTest = true, const float InnerRadius = 10.0f, const float OuterRadius = 15.0f, const int32 Segments = 17);
 
 	/** Debug Draw a triangle on the XY axis */
-	UFUNCTION(BlueprintCallable, Category = "Draw Debug Library", meta = (DevelopmentOnly, AutoCreateRefTerm = "Location, Rotation"))
-	static UE_API void DrawDebugTriangle(const FDebugDrawer& Drawer, const FVector& Location, const FRotator& Rotation, const FLinearColor Color = FLinearColor::Red, const float Thickness = 0.0f, const bool bDepthTest = true, const float Radius = 10.0f);
+	UFUNCTION(BlueprintCallable, Category = "Draw Debug Library", meta = (DevelopmentOnly, AutoCreateRefTerm = "Location, Rotation, LineStyle"))
+	static UE_API void DrawDebugTriangle(const FDebugDrawer& Drawer, const FVector& Location, const FRotator& Rotation, const FDrawDebugLineStyle& LineStyle = FDrawDebugLineStyle(), const bool bDepthTest = true, const float Radius = 10.0f);
 
 	/** Debug Draw a square on the XY axis */
-	UFUNCTION(BlueprintCallable, Category = "Draw Debug Library", meta = (DevelopmentOnly, AutoCreateRefTerm = "Location, Rotation"))
-	static UE_API void DrawDebugSquare(const FDebugDrawer& Drawer, const FVector& Location, const FRotator& Rotation, const FLinearColor Color = FLinearColor::Red, const float Thickness = 0.0f, const bool bDepthTest = true, const float HalfLength = 10.0f);
+	UFUNCTION(BlueprintCallable, Category = "Draw Debug Library", meta = (DevelopmentOnly, AutoCreateRefTerm = "Location, Rotation, LineStyle"))
+	static UE_API void DrawDebugSquare(const FDebugDrawer& Drawer, const FVector& Location, const FRotator& Rotation, const FDrawDebugLineStyle& LineStyle = FDrawDebugLineStyle(), const bool bDepthTest = true, const float HalfLength = 10.0f);
 
 	/** Debug Draw a diamond on the XY axis */
-	UFUNCTION(BlueprintCallable, Category = "Draw Debug Library", meta = (DevelopmentOnly, AutoCreateRefTerm = "Location, Rotation"))
-	static UE_API void DrawDebugDiamond(const FDebugDrawer& Drawer, const FVector& Location, const FRotator& Rotation, const FLinearColor Color = FLinearColor::Red, const float Thickness = 0.0f, const bool bDepthTest = true, const float Radius = 10.0f);
+	UFUNCTION(BlueprintCallable, Category = "Draw Debug Library", meta = (DevelopmentOnly, AutoCreateRefTerm = "Location, Rotation, LineStyle"))
+	static UE_API void DrawDebugDiamond(const FDebugDrawer& Drawer, const FVector& Location, const FRotator& Rotation, const FDrawDebugLineStyle& LineStyle = FDrawDebugLineStyle(), const bool bDepthTest = true, const float Radius = 10.0f);
 
 	/** Debug Draw a pentagon on the XY axis */
-	UFUNCTION(BlueprintCallable, Category = "Draw Debug Library", meta = (DevelopmentOnly, AutoCreateRefTerm = "Location, Rotation"))
-	static UE_API void DrawDebugPentagon(const FDebugDrawer& Drawer, const FVector& Location, const FRotator& Rotation, const FLinearColor Color = FLinearColor::Red, const float Thickness = 0.0f, const bool bDepthTest = true, const float Radius = 10.0f);
+	UFUNCTION(BlueprintCallable, Category = "Draw Debug Library", meta = (DevelopmentOnly, AutoCreateRefTerm = "Location, Rotation, LineStyle"))
+	static UE_API void DrawDebugPentagon(const FDebugDrawer& Drawer, const FVector& Location, const FRotator& Rotation, const FDrawDebugLineStyle& LineStyle = FDrawDebugLineStyle(), const bool bDepthTest = true, const float Radius = 10.0f);
 
 	/** Debug Draw a hexagon on the XY axis */
-	UFUNCTION(BlueprintCallable, Category = "Draw Debug Library", meta = (DevelopmentOnly, AutoCreateRefTerm = "Location, Rotation"))
-	static UE_API void DrawDebugHexagon(const FDebugDrawer& Drawer, const FVector& Location, const FRotator& Rotation, const FLinearColor Color = FLinearColor::Red, const float Thickness = 0.0f, const bool bDepthTest = true, const float Radius = 10.0f);
+	UFUNCTION(BlueprintCallable, Category = "Draw Debug Library", meta = (DevelopmentOnly, AutoCreateRefTerm = "Location, Rotation, LineStyle"))
+	static UE_API void DrawDebugHexagon(const FDebugDrawer& Drawer, const FVector& Location, const FRotator& Rotation, const FDrawDebugLineStyle& LineStyle = FDrawDebugLineStyle(), const bool bDepthTest = true, const float Radius = 10.0f);
 
 	/** Debug Draw a heptagon on the XY axis */
-	UFUNCTION(BlueprintCallable, Category = "Draw Debug Library", meta = (DevelopmentOnly, AutoCreateRefTerm = "Location, Rotation"))
-	static UE_API void DrawDebugHeptagon(const FDebugDrawer& Drawer, const FVector& Location, const FRotator& Rotation, const FLinearColor Color = FLinearColor::Red, const float Thickness = 0.0f, const bool bDepthTest = true, const float Radius = 10.0f);
+	UFUNCTION(BlueprintCallable, Category = "Draw Debug Library", meta = (DevelopmentOnly, AutoCreateRefTerm = "Location, Rotation, LineStyle"))
+	static UE_API void DrawDebugHeptagon(const FDebugDrawer& Drawer, const FVector& Location, const FRotator& Rotation, const FDrawDebugLineStyle& LineStyle = FDrawDebugLineStyle(), const bool bDepthTest = true, const float Radius = 10.0f);
 
 	/** Debug Draw a octagon on the XY axis */
-	UFUNCTION(BlueprintCallable, Category = "Draw Debug Library", meta = (DevelopmentOnly, AutoCreateRefTerm = "Location, Rotation"))
-	static UE_API void DrawDebugOctagon(const FDebugDrawer& Drawer, const FVector& Location, const FRotator& Rotation, const FLinearColor Color = FLinearColor::Red, const float Thickness = 0.0f, const bool bDepthTest = true, const float Radius = 10.0f);
+	UFUNCTION(BlueprintCallable, Category = "Draw Debug Library", meta = (DevelopmentOnly, AutoCreateRefTerm = "Location, Rotation, LineStyle"))
+	static UE_API void DrawDebugOctagon(const FDebugDrawer& Drawer, const FVector& Location, const FRotator& Rotation, const FDrawDebugLineStyle& LineStyle = FDrawDebugLineStyle(), const bool bDepthTest = true, const float Radius = 10.0f);
 
 	/** Debug Draw a nonagon on the XY axis */
-	UFUNCTION(BlueprintCallable, Category = "Draw Debug Library", meta = (DevelopmentOnly, AutoCreateRefTerm = "Location, Rotation"))
-	static UE_API void DrawDebugNonagon(const FDebugDrawer& Drawer, const FVector& Location, const FRotator& Rotation, const FLinearColor Color = FLinearColor::Red, const float Thickness = 0.0f, const bool bDepthTest = true, const float Radius = 10.0f);
+	UFUNCTION(BlueprintCallable, Category = "Draw Debug Library", meta = (DevelopmentOnly, AutoCreateRefTerm = "Location, Rotation, LineStyle"))
+	static UE_API void DrawDebugNonagon(const FDebugDrawer& Drawer, const FVector& Location, const FRotator& Rotation, const FDrawDebugLineStyle& LineStyle = FDrawDebugLineStyle(), const bool bDepthTest = true, const float Radius = 10.0f);
 
 	/** Debug Draw a decagon on the XY axis */
-	UFUNCTION(BlueprintCallable, Category = "Draw Debug Library", meta = (DevelopmentOnly, AutoCreateRefTerm = "Location, Rotation"))
-	static UE_API void DrawDebugDecagon(const FDebugDrawer& Drawer, const FVector& Location, const FRotator& Rotation, const FLinearColor Color = FLinearColor::Red, const float Thickness = 0.0f, const bool bDepthTest = true, const float Radius = 10.0f);
+	UFUNCTION(BlueprintCallable, Category = "Draw Debug Library", meta = (DevelopmentOnly, AutoCreateRefTerm = "Location, Rotation, LineStyle"))
+	static UE_API void DrawDebugDecagon(const FDebugDrawer& Drawer, const FVector& Location, const FRotator& Rotation, const FDrawDebugLineStyle& LineStyle = FDrawDebugLineStyle(), const bool bDepthTest = true, const float Radius = 10.0f);
 
 	/** Debug Draw a flat regular polygon on the XY axis */
-	UFUNCTION(BlueprintCallable, Category = "Draw Debug Library", meta = (DevelopmentOnly, AutoCreateRefTerm = "Location, Rotation"))
-	static UE_API void DrawDebugRegularPolygon(const FDebugDrawer& Drawer, const FVector& Location, const FRotator& Rotation, const FLinearColor Color = FLinearColor::Red, const float Thickness = 0.0f, const bool bDepthTest = true, const float Radius = 10.0f, const int32 Sides = 11);
+	UFUNCTION(BlueprintCallable, Category = "Draw Debug Library", meta = (DevelopmentOnly, AutoCreateRefTerm = "Location, Rotation, LineStyle"))
+	static UE_API void DrawDebugRegularPolygon(const FDebugDrawer& Drawer, const FVector& Location, const FRotator& Rotation, const FDrawDebugLineStyle& LineStyle = FDrawDebugLineStyle(), const bool bDepthTest = true, const float Radius = 10.0f, const int32 Sides = 11);
 
 	/** Debug Draw an axis-aligned cross */
-	UFUNCTION(BlueprintCallable, Category = "Draw Debug Library", meta = (DevelopmentOnly, AutoCreateRefTerm = "Location, Rotation"))
-	static UE_API void DrawDebugLocator(const FDebugDrawer& Drawer, const FVector& Location, const FRotator& Rotation, const FLinearColor Color = FLinearColor::Red, const float Thickness = 0.0f, const bool bDepthTest = true, const float Radius = 10.0f);
+	UFUNCTION(BlueprintCallable, Category = "Draw Debug Library", meta = (DevelopmentOnly, AutoCreateRefTerm = "Location, Rotation, LineStyle"))
+	static UE_API void DrawDebugLocator(const FDebugDrawer& Drawer, const FVector& Location, const FRotator& Rotation, const FDrawDebugLineStyle& LineStyle = FDrawDebugLineStyle(), const bool bDepthTest = true, const float Radius = 10.0f);
 
 	/** Debug Draw an axis-aligned cross rotated by 45 degrees */
-	UFUNCTION(BlueprintCallable, Category = "Draw Debug Library", meta = (DevelopmentOnly, AutoCreateRefTerm = "Location, Rotation"))
-	static UE_API void DrawDebugCrossLocator(const FDebugDrawer& Drawer, const FVector& Location, const FRotator& Rotation, const FLinearColor Color = FLinearColor::Red, const float Thickness = 0.0f, const bool bDepthTest = true, const float Radius = 10.0f);
+	UFUNCTION(BlueprintCallable, Category = "Draw Debug Library", meta = (DevelopmentOnly, AutoCreateRefTerm = "Location, Rotation, LineStyle"))
+	static UE_API void DrawDebugCrossLocator(const FDebugDrawer& Drawer, const FVector& Location, const FRotator& Rotation, const FDrawDebugLineStyle& LineStyle = FDrawDebugLineStyle(), const bool bDepthTest = true, const float Radius = 10.0f);
 
 	/** Debug Draw an oriented box */
-	UFUNCTION(BlueprintCallable, Category = "Draw Debug Library", meta = (DevelopmentOnly, AutoCreateRefTerm = "Location, Rotation"))
-	static UE_API void DrawDebugBox(const FDebugDrawer& Drawer, const FVector& Location, const FRotator& Rotation, const FLinearColor Color = FLinearColor::Red, const float Thickness = 0.0f, const bool bDepthTest = true, const FVector HalfExtents = FVector(10.0f, 10.0f, 10.0f));
+	UFUNCTION(BlueprintCallable, Category = "Draw Debug Library", meta = (DevelopmentOnly, AutoCreateRefTerm = "Location, Rotation, LineStyle"))
+	static UE_API void DrawDebugBox(const FDebugDrawer& Drawer, const FVector& Location, const FRotator& Rotation, const FDrawDebugLineStyle& LineStyle = FDrawDebugLineStyle(), const bool bDepthTest = true, const FVector HalfExtents = FVector(10.0f, 10.0f, 10.0f));
 
 	/** Debug Draw a sphere */
-	UFUNCTION(BlueprintCallable, Category = "Draw Debug Library", meta = (DevelopmentOnly, AutoCreateRefTerm = "Location, Rotation"))
-	static UE_API void DrawDebugSphere(const FDebugDrawer& Drawer, const FVector& Location, const FRotator& Rotation, const FLinearColor Color = FLinearColor::Red, const float Thickness = 0.0f, const bool bDepthTest = true, const float Radius = 10.0f, const int32 Segments = 8);
+	UFUNCTION(BlueprintCallable, Category = "Draw Debug Library", meta = (DevelopmentOnly, AutoCreateRefTerm = "Location, Rotation, LineStyle"))
+	static UE_API void DrawDebugSphere(const FDebugDrawer& Drawer, const FVector& Location, const FRotator& Rotation, const FDrawDebugLineStyle& LineStyle = FDrawDebugLineStyle(), const bool bDepthTest = true, const float Radius = 10.0f, const int32 Segments = 8);
 
 	/** Debug Draw a simple sphere made up of three circles one of each axis */
-	UFUNCTION(BlueprintCallable, Category = "Draw Debug Library", meta = (DevelopmentOnly, AutoCreateRefTerm = "Location, Rotation"))
-	static UE_API void DrawDebugSimpleSphere(const FDebugDrawer& Drawer, const FVector& Location, const FRotator& Rotation, const FLinearColor Color = FLinearColor::Red, const float Thickness = 0.0f, const bool bDepthTest = true, const float Radius = 10.0f, const int32 Segments = 13);
+	UFUNCTION(BlueprintCallable, Category = "Draw Debug Library", meta = (DevelopmentOnly, AutoCreateRefTerm = "Location, Rotation, LineStyle"))
+	static UE_API void DrawDebugSimpleSphere(const FDebugDrawer& Drawer, const FVector& Location, const FRotator& Rotation, const FDrawDebugLineStyle& LineStyle = FDrawDebugLineStyle(), const bool bDepthTest = true, const float Radius = 10.0f, const int32 Segments = 13);
 
 	/** Debug Draw a capsule */
-	UFUNCTION(BlueprintCallable, Category = "Draw Debug Library", meta = (DevelopmentOnly, AutoCreateRefTerm = "Location, Rotation"))
-	static UE_API void DrawDebugCapsule(const FDebugDrawer& Drawer, const FVector& Location, const FRotator& Rotation, const FLinearColor Color = FLinearColor::Red, const float Thickness = 0.0f, const bool bDepthTest = true, const float Radius = 10.0f, const float HalfLength = 10.0f, const int32 Segments = 8);
+	UFUNCTION(BlueprintCallable, Category = "Draw Debug Library", meta = (DevelopmentOnly, AutoCreateRefTerm = "Location, Rotation, LineStyle"))
+	static UE_API void DrawDebugCapsule(const FDebugDrawer& Drawer, const FVector& Location, const FRotator& Rotation, const FDrawDebugLineStyle& LineStyle = FDrawDebugLineStyle(), const bool bDepthTest = true, const float Radius = 10.0f, const float HalfLength = 10.0f, const int32 Segments = 8);
 
 	/** Debug Draw a capsule using the start and end location */
-	UFUNCTION(BlueprintCallable, Category = "Draw Debug Library", meta = (DevelopmentOnly, AutoCreateRefTerm = "StartLocation, EndLocation"))
-	static UE_API void DrawDebugCapsuleLine(const FDebugDrawer& Drawer, const FVector& StartLocation, const FVector& EndLocation, const FLinearColor Color = FLinearColor::Red, const float Thickness = 0.0f, const bool bDepthTest = true, const float Radius = 10.0f, const int32 Segments = 8);
+	UFUNCTION(BlueprintCallable, Category = "Draw Debug Library", meta = (DevelopmentOnly, AutoCreateRefTerm = "StartLocation, EndLocation, LineStyle"))
+	static UE_API void DrawDebugCapsuleLine(const FDebugDrawer& Drawer, const FVector& StartLocation, const FVector& EndLocation, const FDrawDebugLineStyle& LineStyle = FDrawDebugLineStyle(), const bool bDepthTest = true, const float Radius = 10.0f, const int32 Segments = 8);
 
 	/** Debug Draw a Frustum */
-	UFUNCTION(BlueprintCallable, Category = "Draw Debug Library", meta = (DevelopmentOnly, AutoCreateRefTerm = "FrustumToWorld"))
-	static UE_API void DrawDebugFrustum(const FDebugDrawer& Drawer, const FMatrix& FrustumToWorld, const FLinearColor Color = FLinearColor::Red, const float Thickness = 0.0f, const bool bDepthTest = true);
+	UFUNCTION(BlueprintCallable, Category = "Draw Debug Library", meta = (DevelopmentOnly, AutoCreateRefTerm = "FrustumToWorld, LineStyle"))
+	static UE_API void DrawDebugFrustum(const FDebugDrawer& Drawer, const FMatrix& FrustumToWorld, const FDrawDebugLineStyle& LineStyle = FDrawDebugLineStyle(), const bool bDepthTest = true);
 
 public: // Arrows
 
 	/** Debug Draw an arrow */
-	UFUNCTION(BlueprintCallable, Category = "Draw Debug Library", meta = (DevelopmentOnly, AutoCreateRefTerm = "StartLocation, EndLocation, Settings"))
-	static UE_API void DrawDebugArrow(const FDebugDrawer& Drawer, const FVector& StartLocation, const FVector& EndLocation, const FLinearColor Color = FLinearColor::Red, const float Thickness = 0.0f, const bool bDepthTest = true, const FDrawDebugArrowSettings& Settings = FDrawDebugArrowSettings());
+	UFUNCTION(BlueprintCallable, Category = "Draw Debug Library", meta = (DevelopmentOnly, AutoCreateRefTerm = "StartLocation, EndLocation, LineStyle, Settings"))
+	static UE_API void DrawDebugArrow(const FDebugDrawer& Drawer, const FVector& StartLocation, const FVector& EndLocation, const FDrawDebugLineStyle& LineStyle = FDrawDebugLineStyle(), const bool bDepthTest = true, const FDrawDebugArrowSettings& Settings = FDrawDebugArrowSettings());
 
 	/** Debug Draw an arrow with an orientation */
-	UFUNCTION(BlueprintCallable, Category = "Draw Debug Library", meta = (DevelopmentOnly, AutoCreateRefTerm = "Location, Rotation, Settings"))
-	static UE_API void DrawDebugOrientedArrow(const FDebugDrawer& Drawer, const FVector& Location, const FRotator& Rotation, const float Length = 100.0f, const FLinearColor Color = FLinearColor::Red, const float Thickness = 0.0f, const bool bDepthTest = true, const FDrawDebugArrowSettings& Settings = FDrawDebugArrowSettings());
+	UFUNCTION(BlueprintCallable, Category = "Draw Debug Library", meta = (DevelopmentOnly, AutoCreateRefTerm = "Location, Rotation, LineStyle, Settings"))
+	static UE_API void DrawDebugOrientedArrow(const FDebugDrawer& Drawer, const FVector& Location, const FRotator& Rotation, const float Length = 100.0f, const FDrawDebugLineStyle& LineStyle = FDrawDebugLineStyle(), const bool bDepthTest = true, const FDrawDebugArrowSettings& Settings = FDrawDebugArrowSettings());
 
 	/** Draws an arrow pointing down at the location surrounded by a circle */
-	UFUNCTION(BlueprintCallable, Category = "Draw Debug Library", meta = (DevelopmentOnly, AutoCreateRefTerm = "Location"))
-	static UE_API void DrawDebugGroundTargetArrow(const FDebugDrawer& Drawer, const FVector& Location, const FLinearColor Color = FLinearColor::Red, const float Thickness = 0.0f, const bool bDepthTest = true, const float Length = 100.0f, const float ArrowHeadSize = 20.0f, const float Radius = 10.0f, const int32 Segments = 17);
+	UFUNCTION(BlueprintCallable, Category = "Draw Debug Library", meta = (DevelopmentOnly, AutoCreateRefTerm = "Location, LineStyle"))
+	static UE_API void DrawDebugGroundTargetArrow(const FDebugDrawer& Drawer, const FVector& Location, const FDrawDebugLineStyle& LineStyle = FDrawDebugLineStyle(), const bool bDepthTest = true, const float Length = 100.0f, const float ArrowHeadSize = 20.0f, const float Radius = 10.0f, const int32 Segments = 17);
 
 	/** Draws a flat 2d arrow motif */
-	UFUNCTION(BlueprintCallable, Category = "Draw Debug Library", meta = (DevelopmentOnly, AutoCreateRefTerm = "Location, Rotation"))
-	static UE_API void DrawDebugFlatArrow(const FDebugDrawer& Drawer, const FVector& Location, const FRotator& Rotation, const FLinearColor Color = FLinearColor::Red, const float Thickness = 0.0f, const bool bDepthTest = true, const float Length = 20.0f, const float Width = 20.0f);
+	UFUNCTION(BlueprintCallable, Category = "Draw Debug Library", meta = (DevelopmentOnly, AutoCreateRefTerm = "Location, Rotation, LineStyle"))
+	static UE_API void DrawDebugFlatArrow(const FDebugDrawer& Drawer, const FVector& Location, const FRotator& Rotation, const FDrawDebugLineStyle& LineStyle = FDrawDebugLineStyle(), const bool bDepthTest = true, const float Length = 20.0f, const float Width = 20.0f);
 
 	/** Debug Draw an arrow coming off a circle at a given angle */
-	UFUNCTION(BlueprintCallable, Category = "Draw Debug Library", meta = (DevelopmentOnly, AutoCreateRefTerm = "Location, Rotation, ArrowSettings"))
-	static UE_API void DrawDebugCircleArrow(const FDebugDrawer& Drawer, const FVector& Location, const FRotator& Rotation, const float Angle, const FLinearColor Color = FLinearColor::Red, const float Thickness = 0.0f, const bool bDepthTest = true, const float Radius = 10.0f, const float Length = 20.0f, const FDrawDebugArrowSettings& ArrowSettings = FDrawDebugArrowSettings());
+	UFUNCTION(BlueprintCallable, Category = "Draw Debug Library", meta = (DevelopmentOnly, AutoCreateRefTerm = "Location, Rotation, LineStyle, ArrowSettings"))
+	static UE_API void DrawDebugCircleArrow(const FDebugDrawer& Drawer, const FVector& Location, const FRotator& Rotation, const float Angle, const FDrawDebugLineStyle& LineStyle = FDrawDebugLineStyle(), const bool bDepthTest = true, const float Radius = 10.0f, const float Length = 20.0f, const FDrawDebugArrowSettings& ArrowSettings = FDrawDebugArrowSettings());
 
 public: // Splines
 
 	/** Debug Draw the start section of a catmull rom spline */
-	UFUNCTION(BlueprintCallable, Category = "Draw Debug Library", meta = (DevelopmentOnly, AutoCreateRefTerm = "V0, V1, V2"))
-	static UE_API void DrawDebugCatmullRomSplineStart(const FDebugDrawer& Drawer, const FVector& V0, const FVector& V1, const FVector& V2, const FLinearColor Color = FLinearColor::Red, const float Thickness = 0.0f, const bool bDepthTest = true, const bool bMonotonic = false, const int32 Segments = 15);
+	UFUNCTION(BlueprintCallable, Category = "Draw Debug Library", meta = (DevelopmentOnly, AutoCreateRefTerm = "V0, V1, V2, LineStyle"))
+	static UE_API void DrawDebugCatmullRomSplineStart(const FDebugDrawer& Drawer, const FVector& V0, const FVector& V1, const FVector& V2, const FDrawDebugLineStyle& LineStyle = FDrawDebugLineStyle(), const bool bDepthTest = true, const bool bMonotonic = false, const int32 Segments = 15);
 
 	/** Debug Draw the end section of a catmull rom spline */
-	UFUNCTION(BlueprintCallable, Category = "Draw Debug Library", meta = (DevelopmentOnly, AutoCreateRefTerm = "V0, V1, V2"))
-	static UE_API void DrawDebugCatmullRomSplineEnd(const FDebugDrawer& Drawer, const FVector& V0, const FVector& V1, const FVector& V2, const FLinearColor Color = FLinearColor::Red, const float Thickness = 0.0f, const bool bDepthTest = true, const bool bMonotonic = false, const int32 Segments = 15);
+	UFUNCTION(BlueprintCallable, Category = "Draw Debug Library", meta = (DevelopmentOnly, AutoCreateRefTerm = "V0, V1, V2, LineStyle"))
+	static UE_API void DrawDebugCatmullRomSplineEnd(const FDebugDrawer& Drawer, const FVector& V0, const FVector& V1, const FVector& V2, const FDrawDebugLineStyle& LineStyle = FDrawDebugLineStyle(), const bool bDepthTest = true, const bool bMonotonic = false, const int32 Segments = 15);
 
 	/** Debug Draw a full segment of a catmull rom spline */
-	UFUNCTION(BlueprintCallable, Category = "Draw Debug Library", meta = (DevelopmentOnly, AutoCreateRefTerm = "V0, V1, V2, V3"))
-	static UE_API void DrawDebugCatmullRomSplineSection(const FDebugDrawer& Drawer, const FVector& V0, const FVector& V1, const FVector& V2, const FVector& V3, const FLinearColor Color = FLinearColor::Red, const float Thickness = 0.0f, const bool bDepthTest = true, const bool bMonotonic = false, const int32 Segments = 15);
+	UFUNCTION(BlueprintCallable, Category = "Draw Debug Library", meta = (DevelopmentOnly, AutoCreateRefTerm = "V0, V1, V2, V3, LineStyle"))
+	static UE_API void DrawDebugCatmullRomSplineSection(const FDebugDrawer& Drawer, const FVector& V0, const FVector& V1, const FVector& V2, const FVector& V3, const FDrawDebugLineStyle& LineStyle = FDrawDebugLineStyle(), const bool bDepthTest = true, const bool bMonotonic = false, const int32 Segments = 15);
 
 	/** Debug Draw a catmull rom spline */
-	UFUNCTION(BlueprintCallable, Category = "Draw Debug Library", meta = (DevelopmentOnly))
-	static UE_API void DrawDebugCatmullRomSpline(const FDebugDrawer& Drawer, const TArray<FVector>& Points, const FLinearColor Color = FLinearColor::Red, const float Thickness = 0.0f, const bool bDepthTest = true, const bool bMonotonic = false, const int32 Segments = 15);
-	static UE_API void DrawDebugCatmullRomSplineArrayView(const FDebugDrawer& Drawer, const TArrayView<const FVector> Points, const FLinearColor Color = FLinearColor::Red, const float Thickness = 0.0f, const bool bDepthTest = true, const bool bMonotonic = false, const int32 Segments = 15);
+	UFUNCTION(BlueprintCallable, Category = "Draw Debug Library", meta = (DevelopmentOnly, AutoCreateRefTerm = "Points, LineStyle"))
+	static UE_API void DrawDebugCatmullRomSpline(const FDebugDrawer& Drawer, const TArray<FVector>& Points, const FDrawDebugLineStyle& LineStyle = FDrawDebugLineStyle(), const bool bDepthTest = true, const bool bMonotonic = false, const int32 Segments = 15);
+	static UE_API void DrawDebugCatmullRomSplineArrayView(const FDebugDrawer& Drawer, const TArrayView<const FVector> Points, const FDrawDebugLineStyle& LineStyle = FDrawDebugLineStyle(), const bool bDepthTest = true, const bool bMonotonic = false, const int32 Segments = 15);
 
 public: // Variable Types
 
 	/** Debug Draw an angle */
-	UFUNCTION(BlueprintCallable, Category = "Draw Debug Library", meta = (DevelopmentOnly, AutoCreateRefTerm = "Location, Rotation"))
-	static UE_API void DrawDebugAngle(const FDebugDrawer& Drawer, const FVector& Location, const FRotator& Rotation, const float Angle = 0.0f, const FLinearColor Color = FLinearColor::Red, const float Thickness = 0.0f, const bool bDepthTest = true, const float LineLength = 10.0f, const float AngleRadius = 5.0f, const int32 Segments = 17);
+	UFUNCTION(BlueprintCallable, Category = "Draw Debug Library", meta = (DevelopmentOnly, AutoCreateRefTerm = "Location, Rotation, LineStyle"))
+	static UE_API void DrawDebugAngle(const FDebugDrawer& Drawer, const FVector& Location, const FRotator& Rotation, const float Angle = 0.0f, const FDrawDebugLineStyle& LineStyle = FDrawDebugLineStyle(), const bool bDepthTest = true, const float LineLength = 10.0f, const float AngleRadius = 5.0f, const int32 Segments = 17);
 
 	/** Debug Draw a sphere at the given location */
-	UFUNCTION(BlueprintCallable, Category = "Draw Debug Library", meta = (DevelopmentOnly, AutoCreateRefTerm = "Location"))
-	static UE_API void DrawDebugLocation(const FDebugDrawer& Drawer, const FVector& Location, const FLinearColor Color = FLinearColor::Red, const float Thickness = 0.0f, const bool bDepthTest = true, const float DrawRadius = 10.0f, const int32 Segments = 8);
+	UFUNCTION(BlueprintCallable, Category = "Draw Debug Library", meta = (DevelopmentOnly, AutoCreateRefTerm = "Location, LineStyle"))
+	static UE_API void DrawDebugLocation(const FDebugDrawer& Drawer, const FVector& Location, const FDrawDebugLineStyle& LineStyle = FDrawDebugLineStyle(), const bool bDepthTest = true, const float DrawRadius = 10.0f, const int32 Segments = 8);
 
 	/** Debug Draw spheres at the given locations */
-	UFUNCTION(BlueprintCallable, Category = "Draw Debug Library", meta = (DevelopmentOnly))
-	static UE_API void DrawDebugLocations(const FDebugDrawer& Drawer, const TArray<FVector>& Locations, const FLinearColor Color = FLinearColor::Red, const float Thickness = 0.0f, const bool bDepthTest = true, const float DrawRadius = 10.0f, const int32 Segments = 8);
-	static UE_API void DrawDebugLocationsArrayView(const FDebugDrawer& Drawer, const TArrayView<const FVector> Locations, const FLinearColor Color = FLinearColor::Red, const float Thickness = 0.0f, const bool bDepthTest = true, const float DrawRadius = 10.0f, const int32 Segments = 8);
+	UFUNCTION(BlueprintCallable, Category = "Draw Debug Library", meta = (DevelopmentOnly, AutoCreateRefTerm = "Locations, LineStyle"))
+	static UE_API void DrawDebugLocations(const FDebugDrawer& Drawer, const TArray<FVector>& Locations, const FDrawDebugLineStyle& LineStyle = FDrawDebugLineStyle(), const bool bDepthTest = true, const float DrawRadius = 10.0f, const int32 Segments = 8);
+	static UE_API void DrawDebugLocationsArrayView(const FDebugDrawer& Drawer, const TArrayView<const FVector> Locations, const FDrawDebugLineStyle& LineStyle = FDrawDebugLineStyle(), const bool bDepthTest = true, const float DrawRadius = 10.0f, const int32 Segments = 8);
 
 	/** Draws a rotation */
-	UFUNCTION(BlueprintCallable, Category = "Draw Debug Library", meta = (DevelopmentOnly, AutoCreateRefTerm = "Location, Rotation"))
-	static UE_API void DrawDebugRotation(const FDebugDrawer& Drawer, const FVector& Location, const FRotator& Rotation, const float Opacity = 1.0f, const float Thickness = 0.0f, const bool bDepthTest = true, const float DrawRadius = 10.0f);
+	UFUNCTION(BlueprintCallable, Category = "Draw Debug Library", meta = (DevelopmentOnly, AutoCreateRefTerm = "Location, Rotation, LineStyle"))
+	static UE_API void DrawDebugRotation(const FDebugDrawer& Drawer, const FVector& Location, const FRotator& Rotation, const FDrawDebugLineStyle& LineStyle = FDrawDebugLineStyle(), const bool bDepthTest = true, const float DrawRadius = 10.0f);
 
 	/** Draws an array of transforms */
 	static UE_API void DrawDebugRotationsQuatArrayView(
 		const FDebugDrawer& Drawer,
 		const TArrayView<const FVector> Locations,
 		const TArrayView<const FQuat4f> Rotations,
-		const float Opacity = 1.0f,
-		const float Thickness = 0.0f,
+		const FDrawDebugLineStyle& LineStyle = FDrawDebugLineStyle(),
 		const bool bDepthTest = true,
 		const float DrawRadius = 10.0f);
 
 	/** Debug Draw an arrow at the given location facing in the given direction */
-	UFUNCTION(BlueprintCallable, Category = "Draw Debug Library", meta = (DevelopmentOnly, AutoCreateRefTerm = "Location, Direction"))
-	static UE_API void DrawDebugDirection(const FDebugDrawer& Drawer, const FVector& Location, const FVector& Direction = FVector::ForwardVector, const FLinearColor Color = FLinearColor::Red, const float Thickness = 0.0f, const bool bDepthTest = true, const float DrawArrowLength = 100.0f, const float ArrowHeadScale = 1.0f);
+	UFUNCTION(BlueprintCallable, Category = "Draw Debug Library", meta = (DevelopmentOnly, AutoCreateRefTerm = "Location, Direction, LineStyle"))
+	static UE_API void DrawDebugDirection(const FDebugDrawer& Drawer, const FVector& Location, const FVector& Direction = FVector::ForwardVector, const FDrawDebugLineStyle& LineStyle = FDrawDebugLineStyle(), const bool bDepthTest = true, const float DrawArrowLength = 100.0f, const float ArrowHeadScale = 1.0f);
 
 	/** Debug Draw a line at the given location scaled by the given velocity */
-	UFUNCTION(BlueprintCallable, Category = "Draw Debug Library", meta = (DevelopmentOnly, AutoCreateRefTerm = "Location, Velocity"))
-	static UE_API void DrawDebugVelocity(const FDebugDrawer& Drawer, const FVector& Location, const FVector& Velocity, const FLinearColor Color = FLinearColor::Red, const float Thickness = 0.0f, const bool bDepthTest = true, const float DrawVelocityLineScale = 1.0f);
+	UFUNCTION(BlueprintCallable, Category = "Draw Debug Library", meta = (DevelopmentOnly, AutoCreateRefTerm = "Location, Velocity, LineStyle"))
+	static UE_API void DrawDebugVelocity(const FDebugDrawer& Drawer, const FVector& Location, const FVector& Velocity, const FDrawDebugLineStyle& LineStyle = FDrawDebugLineStyle(), const bool bDepthTest = true, const float DrawVelocityLineScale = 1.0f);
 
 	/** Debug draw velocities */
-	UFUNCTION(BlueprintCallable, Category = "Draw Debug Library", meta = (DevelopmentOnly))
-	static UE_API void DrawDebugVelocities(const FDebugDrawer& Drawer, const TArray<FVector>& Locations, const TArray<FVector>& Velocities, const FLinearColor Color = FLinearColor::Red, const float Thickness = 0.0f, const bool bDepthTest = true, const float DrawVelocityLineScale = 1.0f);
-	static UE_API void DrawDebugVelocitiesArrayView(const FDebugDrawer& Drawer, const TArrayView<const FVector> Locations, const TArrayView<const FVector> Velocities, const FLinearColor Color = FLinearColor::Red, const float Thickness = 0.0f, const bool bDepthTest = true, const float DrawVelocityLineScale = 1.0f);
+	UFUNCTION(BlueprintCallable, Category = "Draw Debug Library", meta = (DevelopmentOnly, AutoCreateRefTerm = "Locations, Velocities, LineStyle"))
+	static UE_API void DrawDebugVelocities(const FDebugDrawer& Drawer, const TArray<FVector>& Locations, const TArray<FVector>& Velocities, const FDrawDebugLineStyle& LineStyle = FDrawDebugLineStyle(), const bool bDepthTest = true, const float DrawVelocityLineScale = 1.0f);
+	static UE_API void DrawDebugVelocitiesArrayView(const FDebugDrawer& Drawer, const TArrayView<const FVector> Locations, const TArrayView<const FVector> Velocities, const FDrawDebugLineStyle& LineStyle = FDrawDebugLineStyle(), const bool bDepthTest = true, const float DrawVelocityLineScale = 1.0f);
 
 	/** Debug Draw a set of axes at the given transform */
-	UFUNCTION(BlueprintCallable, Category = "Draw Debug Library", meta = (DevelopmentOnly, AutoCreateRefTerm = "Transform"))
-	static UE_API void DrawDebugTransform(const FDebugDrawer& Drawer, const FTransform& Transform, const float Opacity = 1.0f, const float Thickness = 0.0f, const bool bDepthTest = true, const float DrawRadius = 10.0f);
+	UFUNCTION(BlueprintCallable, Category = "Draw Debug Library", meta = (DevelopmentOnly, AutoCreateRefTerm = "Transform, LineStyle"))
+	static UE_API void DrawDebugTransform(const FDebugDrawer& Drawer, const FTransform& Transform, const FDrawDebugLineStyle& LineStyle = FDrawDebugLineStyle(), const bool bDepthTest = true, const float DrawRadius = 10.0f);
 
 	/** Draws an array of transforms */
 	static UE_API void DrawDebugTransformsArrayView(
@@ -693,37 +825,36 @@ public: // Variable Types
 		const TArrayView<const FVector> Locations,
 		const TArrayView<const FQuat4f> Rotations,
 		const TArrayView<const FVector3f> Scales,
-		const float Opacity = 1.0f,
-		const float Thickness = 0.0f,
+		const FDrawDebugLineStyle& LineStyle = FDrawDebugLineStyle(),
 		const bool bDepthTest = true,
 		const float DrawRadius = 10.0f);
 
 	/** Debug draw a phase-representation of an event at a location */
-	UFUNCTION(BlueprintCallable, Category = "Draw Debug Library", meta = (DevelopmentOnly, AutoCreateRefTerm = "Location"))
-	static UE_API void DrawDebugEvent(const FDebugDrawer& Drawer, bool bTimeUntilEventKnown, const float TimeUntilEvent, const FVector& Location, const FLinearColor Color = FLinearColor::Red, const float Thickness = 0.0f, const bool bDepthTest = true, const float Size = 20.0f);
+	UFUNCTION(BlueprintCallable, Category = "Draw Debug Library", meta = (DevelopmentOnly, AutoCreateRefTerm = "Location, LineStyle"))
+	static UE_API void DrawDebugEvent(const FDebugDrawer& Drawer, bool bTimeUntilEventKnown, const float TimeUntilEvent, const FVector& Location, const FDrawDebugLineStyle& LineStyle = FDrawDebugLineStyle(), const bool bDepthTest = true, const float Size = 20.0f);
 
 public: // Complex Objects
 
 	/** Debug draw a basic chair shape */
-	UFUNCTION(BlueprintCallable, Category = "Draw Debug Library", meta = (DevelopmentOnly, AutoCreateRefTerm = "Location, Rotation, Settings"))
-	static UE_API void DrawDebugChair(const FDebugDrawer& Drawer, const FVector& Location, const FRotator& Rotation, const FLinearColor Color = FLinearColor::Red, const float Thickness = 0.0f, const bool bDepthTest = true, const FDrawDebugChairSettings& Settings = FDrawDebugChairSettings());
+	UFUNCTION(BlueprintCallable, Category = "Draw Debug Library", meta = (DevelopmentOnly, AutoCreateRefTerm = "Location, Rotation, LineStyle, Settings"))
+	static UE_API void DrawDebugChair(const FDebugDrawer& Drawer, const FVector& Location, const FRotator& Rotation, const FDrawDebugLineStyle& LineStyle = FDrawDebugLineStyle(), const bool bDepthTest = true, const FDrawDebugChairSettings& Settings = FDrawDebugChairSettings());
 
 	/** Debug draw a basic door shape */
-	UFUNCTION(BlueprintCallable, Category = "Draw Debug Library", meta = (DevelopmentOnly, AutoCreateRefTerm = "Location, Rotation, Settings"))
-	static UE_API void DrawDebugDoor(const FDebugDrawer& Drawer, const FVector& Location, const FRotator& Rotation, const FLinearColor Color = FLinearColor::Red, const float Thickness = 0.0f, const bool bDepthTest = true, const FDrawDebugDoorSettings& Settings = FDrawDebugDoorSettings());
+	UFUNCTION(BlueprintCallable, Category = "Draw Debug Library", meta = (DevelopmentOnly, AutoCreateRefTerm = "Location, Rotation, LineStyle, Settings"))
+	static UE_API void DrawDebugDoor(const FDebugDrawer& Drawer, const FVector& Location, const FRotator& Rotation, const FDrawDebugLineStyle& LineStyle = FDrawDebugLineStyle(), const bool bDepthTest = true, const FDrawDebugDoorSettings& Settings = FDrawDebugDoorSettings());
 
 	/** Debug draw a camera shape */
-	UFUNCTION(BlueprintCallable, Category = "Draw Debug Library", meta = (DevelopmentOnly, AutoCreateRefTerm = "Location, Rotation"))
-	static UE_API void DrawDebugCamera(const FDebugDrawer& Drawer, const FVector& Location, const FRotator& Rotation, const FLinearColor Color = FLinearColor::Red, const float Thickness = 0.0f, const bool bDepthTest = true, const float Scale = 10.0f, const float FOVDegees = 30.0f);
+	UFUNCTION(BlueprintCallable, Category = "Draw Debug Library", meta = (DevelopmentOnly, AutoCreateRefTerm = "Location, Rotation, LineStyle"))
+	static UE_API void DrawDebugCamera(const FDebugDrawer& Drawer, const FVector& Location, const FRotator& Rotation, const FDrawDebugLineStyle& LineStyle = FDrawDebugLineStyle(), const bool bDepthTest = true, const float Scale = 10.0f, const float FOVDegees = 30.0f);
 
 	/** Debug Draw a trajectory of locations and directions */
-	UFUNCTION(BlueprintCallable, Category = "Draw Debug Library", meta = (DevelopmentOnly, AutoCreateRefTerm = "RelativeTransform"))
-	static UE_API void DrawDebugTrajectory(const FDebugDrawer& Drawer, const TArray<FVector>& Locations, const TArray<FVector>& Directions, const FTransform& RelativeTransform = FTransform(), const FLinearColor Color = FLinearColor::Red, const float Thickness = 0.0f, const bool bDepthTest = true, const float DrawArrowLength = 100.0f, const float PointRadius = 5.0f, const float ArrowHeadScale = 1.0f, const int32 Segments = 9, const float VerticalOffset = 2.0f);
-	static UE_API void DrawDebugTrajectoryFromArrayViews(const FDebugDrawer& Drawer, const TArrayView<const FVector> Locations, const TArrayView<const FVector> Directions, const FTransform& RelativeTransform = FTransform(), const FLinearColor Color = FLinearColor::Red, const float Thickness = 0.0f, const bool bDepthTest = true, const float DrawArrowLength = 100.0f, const float PointRadius = 5.0f, const float ArrowHeadScale = 1.0f, const int32 Segments = 9, const float VerticalOffset = 2.0f);
+	UFUNCTION(BlueprintCallable, Category = "Draw Debug Library", meta = (DevelopmentOnly, AutoCreateRefTerm = "Locations, Directions, RelativeTransform, LineStyle"))
+	static UE_API void DrawDebugTrajectory(const FDebugDrawer& Drawer, const TArray<FVector>& Locations, const TArray<FVector>& Directions, const FTransform& RelativeTransform = FTransform(), const FDrawDebugLineStyle& LineStyle = FDrawDebugLineStyle(), const bool bDepthTest = true, const float DrawArrowLength = 100.0f, const float PointRadius = 5.0f, const float ArrowHeadScale = 1.0f, const int32 Segments = 9, const float VerticalOffset = 2.0f);
+	static UE_API void DrawDebugTrajectoryFromArrayViews(const FDebugDrawer& Drawer, const TArrayView<const FVector> Locations, const TArrayView<const FVector> Directions, const FTransform& RelativeTransform = FTransform(), const FDrawDebugLineStyle& LineStyle = FDrawDebugLineStyle(), const bool bDepthTest = true, const float DrawArrowLength = 100.0f, const float PointRadius = 5.0f, const float ArrowHeadScale = 1.0f, const int32 Segments = 9, const float VerticalOffset = 2.0f);
 
 	/** Debug Draw a transform trajectory */
-	UFUNCTION(BlueprintCallable, Category = "Draw Debug Library", meta = (DevelopmentOnly))
-	static UE_API void DrawDebugTransformTrajectory(const FDebugDrawer& Drawer, UPARAM(ref) const FTransformTrajectory& TransformTrajectory, const FLinearColor Color = FLinearColor::Red, const float Thickness = 0.0f, const bool bDepthTest = true, const float Radius = 5.0f, const float VerticalOffset = 2.0f);
+	UFUNCTION(BlueprintCallable, Category = "Draw Debug Library", meta = (DevelopmentOnly, AutoCreateRefTerm = "TransformTrajectory, LineStyle"))
+	static UE_API void DrawDebugTransformTrajectory(const FDebugDrawer& Drawer, UPARAM(ref) const FTransformTrajectory& TransformTrajectory, const FDrawDebugLineStyle& LineStyle = FDrawDebugLineStyle(), const bool bDepthTest = true, const float Radius = 5.0f, const float VerticalOffset = 2.0f);
 
 	/** Draws a trajectory in the style of trajectory drawing in persona */
 	static UE_API void DrawDebugRangeTrajectoryArrayView(
@@ -737,8 +868,8 @@ public: // Complex Objects
 		const FVector& OriginOffset = FVector::ZeroVector);
 
 	/** Draws a flat 2d motif representing mover position and rotation */
-	UFUNCTION(BlueprintCallable, Category = "Draw Debug Library", meta = (DevelopmentOnly, AutoCreateRefTerm = "Location, Rotation"))
-	static UE_API void DrawDebugMoverOrientation(const FDebugDrawer& Drawer, const FVector& Location, const FRotator& Rotation, const FLinearColor Color = FLinearColor::Red, const float Thickness = 0.0f, const bool bDepthTest = true, const FVector ForwardVector = FVector::ForwardVector, const float Scale = 30.0f);
+	UFUNCTION(BlueprintCallable, Category = "Draw Debug Library", meta = (DevelopmentOnly, AutoCreateRefTerm = "Location, Rotation, LineStyle"))
+	static UE_API void DrawDebugMoverOrientation(const FDebugDrawer& Drawer, const FVector& Location, const FRotator& Rotation, const FDrawDebugLineStyle& LineStyle = FDrawDebugLineStyle(), const bool bDepthTest = true, const FVector ForwardVector = FVector::ForwardVector, const float Scale = 30.0f);
 
 public: // Poses / Skeletons
 
@@ -752,11 +883,11 @@ public: // Poses / Skeletons
 
 	/** Debug Draw a bone */
 	UFUNCTION(BlueprintCallable, Category = "Draw Debug Library", meta = (DevelopmentOnly, AutoCreateRefTerm = "Location, Rotation"))
-	static UE_API void DrawDebugBone(const FDebugDrawer& Drawer, const FVector& Location, const FRotator& Rotation, const FLinearColor Color = FLinearColor(0.0f, 0.0f, 0.025f, 1.0f), const float Thickness = 0.0f, const bool bDepthTest = true, const float Radius = 5.0f, const int32 Segments = 10, const bool bDrawTransform = false);
+	static UE_API void DrawDebugBone(const FDebugDrawer& Drawer, const FVector& Location, const FRotator& Rotation, const FLinearColor Color = FLinearColor(0.0f, 0.0f, 0.025f, 1.0f), const bool bDepthTest = true, const float Radius = 5.0f, const int32 Segments = 10, const bool bDrawTransform = false);
 
 	/** Debug Draw a link between a child and parent bone */
 	UFUNCTION(BlueprintCallable, Category = "Draw Debug Library", meta = (DevelopmentOnly, AutoCreateRefTerm = "ChildLocation, ParentLocation, ParentRotation"))
-	static UE_API void DrawDebugBoneLink(const FDebugDrawer& Drawer, const FVector& ChildLocation, const FVector& ParentLocation, const FRotator& ParentRotation, const FLinearColor Color = FLinearColor(0.0f, 0.0f, 0.025f, 1.0f), const float Thickness = 0.0f, const bool bDepthTest = true, const float Radius = 5.0f);
+	static UE_API void DrawDebugBoneLink(const FDebugDrawer& Drawer, const FVector& ChildLocation, const FVector& ParentLocation, const FRotator& ParentRotation, const FLinearColor Color = FLinearColor(0.0f, 0.0f, 0.025f, 1.0f), const bool bDepthTest = true, const float Radius = 5.0f);
 
 	/** Draws a skeleton from a skinned mesh component */
 	UFUNCTION(BlueprintCallable, Category = "Draw Debug Library", meta = (DevelopmentOnly))
@@ -764,7 +895,6 @@ public: // Poses / Skeletons
 		const FDebugDrawer& Drawer,
 		const USkinnedMeshComponent* SkinnedMeshComponent,
 		const FLinearColor Color = FLinearColor(0.0f, 0.0f, 0.025f, 1.0f),
-		const float Thickness = 0.0f,
 		const bool bDepthTest = false,
 		const FDrawDebugSkeletonSettings& Settings = FDrawDebugSkeletonSettings());
 
@@ -778,7 +908,6 @@ public: // Poses / Skeletons
 		const TArrayView<const int32> BoneIndices,
 		const TArrayView<const int32> BoneParents,
 		const FLinearColor Color = FLinearColor(0.0f, 0.0f, 0.025f, 1.0f),
-		const float Thickness = 0.0f,
 		const bool bDepthTest = false,
 		const FDrawDebugSkeletonSettings& Settings = FDrawDebugSkeletonSettings());
 
@@ -788,15 +917,15 @@ public: // Poses / Skeletons
 		const TArrayView<const FVector> BoneLocations,
 		const TArrayView<const FVector3f> BoneVelocities,
 		const TArrayView<const int32> BoneIndices,
-		const FLinearColor& Color = FLinearColor(0.0f, 0.0f, 0.025f, 1.0f),
-		const float Thickness = 0.0f,
+		const FLinearColor Color = FLinearColor(0.0f, 0.0f, 0.025f, 1.0f),
+		const float Thickness = 1.0f,
 		const bool bDepthTest = false,
 		const float Scale = 100.0f);
 
 	/** Debug Draw a pose made up of bone locations and velocities */
-	UFUNCTION(BlueprintCallable, Category = "Draw Debug Library", meta = (DevelopmentOnly, AutoCreateRefTerm = "RelativeTransform"))
-	static UE_API void DrawDebugPose(const FDebugDrawer& Drawer, const TArray<FVector>& BoneLocations, const TArray<FVector>& BoneLinearVelocities, const FTransform& RelativeTransform = FTransform(), const FLinearColor Color = FLinearColor::Red, const float Thickness = 0.0f, const bool bDepthTest = true, const float DrawVelocityLineScale = 1.0f);
-	static UE_API void DrawDebugPoseFromArrayViews(const FDebugDrawer& Drawer, const TArrayView<const FVector> BoneLocations, const TArrayView<const FVector> BoneLinearVelocities, const FTransform& RelativeTransform = FTransform(), const FLinearColor Color = FLinearColor::Red, const float Thickness = 0.0f, const bool bDepthTest = true, const float DrawVelocityLineScale = 1.0f);
+	UFUNCTION(BlueprintCallable, Category = "Draw Debug Library", meta = (DevelopmentOnly, AutoCreateRefTerm = "BoneLocations, BoneLinearVelocities, RelativeTransform, LineStyle"))
+	static UE_API void DrawDebugPose(const FDebugDrawer& Drawer, const TArray<FVector>& BoneLocations, const TArray<FVector>& BoneLinearVelocities, const FTransform& RelativeTransform = FTransform(), const FDrawDebugLineStyle& LineStyle = FDrawDebugLineStyle(), const bool bDepthTest = true, const float DrawVelocityLineScale = 1.0f);
+	static UE_API void DrawDebugPoseFromArrayViews(const FDebugDrawer& Drawer, const TArrayView<const FVector> BoneLocations, const TArrayView<const FVector> BoneLinearVelocities, const FTransform& RelativeTransform = FTransform(), const FDrawDebugLineStyle& LineStyle = FDrawDebugLineStyle(), const bool bDepthTest = true, const float DrawVelocityLineScale = 1.0f);
 
 public: // Text
 
@@ -819,38 +948,37 @@ public: // Text
 	static UE_API FVector DrawDebugNameDimensions(const FName& Name, const FDrawDebugStringSettings& Settings = FDrawDebugStringSettings());
 
 	/** Debug Draw a string. Will only render ASCII characters. */
-	UFUNCTION(BlueprintCallable, Category = "Draw Debug Library", meta = (DevelopmentOnly, AutoCreateRefTerm = "String, Location, Rotation, Settings"))
-	static UE_API void DrawDebugString(const FDebugDrawer& Drawer, const FString& String, const FVector& Location, const FRotator& Rotation, const FLinearColor Color = FLinearColor::Red, const float Thickness = 0.0f, const bool bDepthTest = true, const FDrawDebugStringSettings& Settings = FDrawDebugStringSettings());
-	static UE_API void DrawDebugStringView(const FDebugDrawer& Drawer, const FStringView String, const FVector& Location, const FRotator& Rotation, const FLinearColor Color = FLinearColor::Red, const float Thickness = 0.0f, const bool bDepthTest = true, const FDrawDebugStringSettings& Settings = FDrawDebugStringSettings());
+	UFUNCTION(BlueprintCallable, Category = "Draw Debug Library", meta = (DevelopmentOnly, AutoCreateRefTerm = "String, Location, Rotation, LineStyle, Settings"))
+	static UE_API void DrawDebugString(const FDebugDrawer& Drawer, const FString& String, const FVector& Location, const FRotator& Rotation, const FDrawDebugLineStyle& LineStyle = FDrawDebugLineStyle(), const bool bDepthTest = true, const FDrawDebugStringSettings& Settings = FDrawDebugStringSettings());
+	static UE_API void DrawDebugStringView(const FDebugDrawer& Drawer, const FStringView String, const FVector& Location, const FRotator& Rotation, const FDrawDebugLineStyle& LineStyle = FDrawDebugLineStyle(), const bool bDepthTest = true, const FDrawDebugStringSettings& Settings = FDrawDebugStringSettings());
 
 	/** Debug Draw a name */
-	UFUNCTION(BlueprintCallable, Category = "Draw Debug Library", meta = (DevelopmentOnly, AutoCreateRefTerm = "Name, Location, Rotation, Settings"))
-	static UE_API void DrawDebugName(const FDebugDrawer& Drawer, const FName& Name, const FVector& Location, const FRotator& Rotation, const FLinearColor Color = FLinearColor::Red, const float Thickness = 0.0f, const bool bDepthTest = true, const FDrawDebugStringSettings& Settings = FDrawDebugStringSettings());
+	UFUNCTION(BlueprintCallable, Category = "Draw Debug Library", meta = (DevelopmentOnly, AutoCreateRefTerm = "Name, Location, Rotation, LineStyle, Settings"))
+	static UE_API void DrawDebugName(const FDebugDrawer& Drawer, const FName& Name, const FVector& Location, const FRotator& Rotation, const FDrawDebugLineStyle& LineStyle = FDrawDebugLineStyle(), const bool bDepthTest = true, const FDrawDebugStringSettings& Settings = FDrawDebugStringSettings());
 
 	/** Debug Draw some text. Will only render ASCII characters. */
-	UFUNCTION(BlueprintCallable, Category = "Draw Debug Library", meta = (DevelopmentOnly, AutoCreateRefTerm = "Text, Location, Rotation, Settings"))
-	static UE_API void DrawDebugText(const FDebugDrawer& Drawer, const FText& Text, const FVector& Location, const FRotator& Rotation, const FLinearColor Color = FLinearColor::Red, const float Thickness = 0.0f, const bool bDepthTest = true, const FDrawDebugStringSettings& Settings = FDrawDebugStringSettings());
+	UFUNCTION(BlueprintCallable, Category = "Draw Debug Library", meta = (DevelopmentOnly, AutoCreateRefTerm = "Text, Location, Rotation, LineStyle, Settings"))
+	static UE_API void DrawDebugText(const FDebugDrawer& Drawer, const FText& Text, const FVector& Location, const FRotator& Rotation, const FDrawDebugLineStyle& LineStyle = FDrawDebugLineStyle(), const bool bDepthTest = true, const FDrawDebugStringSettings& Settings = FDrawDebugStringSettings());
 
 public: // Visual Logger Text
 
 	/** Debug Draw a string to the visual logger. Will do nothing if not using a VisualLoggerDrawer. Will not display on-screen during recording. */
 	UFUNCTION(BlueprintCallable, Category = "Draw Debug Library", meta = (DevelopmentOnly, AutoCreateRefTerm = "String, Location"))
-	static UE_API void VisualLoggerDrawString(const FDebugDrawer& Drawer, const FString& String, const FVector& Location, const FLinearColor Color = FLinearColor::Red);
-	static UE_API void VisualLoggerDrawStringView(const FDebugDrawer& Drawer, const FStringView String, const FVector& Location, const FLinearColor Color = FLinearColor::Red);
+	static UE_API void VisualLoggerDrawString(const FDebugDrawer& Drawer, const FString& String, const FVector& Location, const FLinearColor Color = FLinearColor(1.0f, 0.0f, 1.0f, 1.0f));
+	static UE_API void VisualLoggerDrawStringView(const FDebugDrawer& Drawer, const FStringView String, const FVector& Location, const FLinearColor Color = FLinearColor(1.0f, 0.0f, 1.0f, 1.0f));
 
 	/** Debug Draw a name to the visual logger. Will do nothing if not using a VisualLoggerDrawer. Will not display on-screen during recording. */
 	UFUNCTION(BlueprintCallable, Category = "Draw Debug Library", meta = (DevelopmentOnly, AutoCreateRefTerm = "Name, Location"))
-	static UE_API void VisualLoggerDrawName(const FDebugDrawer& Drawer, const FName& Name, const FVector& Location, const FLinearColor Color = FLinearColor::Red);
+	static UE_API void VisualLoggerDrawName(const FDebugDrawer& Drawer, const FName& Name, const FVector& Location, const FLinearColor Color = FLinearColor(1.0f, 0.0f, 1.0f, 1.0f));
 
 	/** Debug Draw some text to the visual logger. Will do nothing if not using a VisualLoggerDrawer. Will not display on-screen during recording. */
 	UFUNCTION(BlueprintCallable, Category = "Draw Debug Library", meta = (DevelopmentOnly, AutoCreateRefTerm = "Text, Location"))
-	static UE_API void VisualLoggerDrawText(const FDebugDrawer& Drawer, const FText& Text, const FVector& Location, const FLinearColor Color = FLinearColor::Red);
-
+	static UE_API void VisualLoggerDrawText(const FDebugDrawer& Drawer, const FText& Text, const FVector& Location, const FLinearColor Color = FLinearColor(1.0f, 0.0f, 1.0f, 1.0f));
 
 public: // Graphs
 
 	/** Debug Draw a simple graph. */
-	UFUNCTION(BlueprintCallable, Category = "Draw Debug Library", meta = (DevelopmentOnly, AutoCreateRefTerm = "Location, Rotation, AxesSettings"))
+	UFUNCTION(BlueprintCallable, Category = "Draw Debug Library", meta = (DevelopmentOnly, AutoCreateRefTerm = "Location, Rotation, TextLineStyle, AxesLineStyle, PlotLineStyle, AxesSettings"))
 	static UE_API void DrawDebugGraph(
 		const FDebugDrawer& Drawer,
 		const FVector& Location,
@@ -863,10 +991,9 @@ public: // Graphs
 		const float Ymax = 1.0f,
 		const float XaxisLength = 100.0f,
 		const float YaxisLength = 100.0f,
-		const FLinearColor Color = FLinearColor::Red,
-		const float AxisThickness = 1.0f,
-		const float LineThickness = 0.0f,
-		const float TextThickness = 0.0f,
+		const FDrawDebugLineStyle& TextLineStyle = FDrawDebugLineStyle(),
+		const FDrawDebugLineStyle& AxesLineStyle = FDrawDebugLineStyle(),
+		const FDrawDebugLineStyle& PlotLineStyle = FDrawDebugLineStyle(),
 		const bool bDepthTest = true,
 		const FDrawDebugGraphAxesSettings& AxesSettings = FDrawDebugGraphAxesSettings());
 
@@ -882,29 +1009,27 @@ public: // Graphs
 		const float Ymax = 1.0f,
 		const float XaxisLength = 100.0f,
 		const float YaxisLength = 100.0f,
-		const FLinearColor Color = FLinearColor::Red,
-		const float AxisThickness = 1.0f,
-		const float LineThickness = 0.0f,
-		const float TextThickness = 0.0f,
+		const FDrawDebugLineStyle& TextLineStyle = FDrawDebugLineStyle(),
+		const FDrawDebugLineStyle& AxesLineStyle = FDrawDebugLineStyle(),
+		const FDrawDebugLineStyle& PlotLineStyle = FDrawDebugLineStyle(),
 		const bool bDepthTest = true,
 		const FDrawDebugGraphAxesSettings& AxesSettings = FDrawDebugGraphAxesSettings());
 
 	/** Debug Draw a simple graph axes. */
-	UFUNCTION(BlueprintCallable, Category = "Draw Debug Library", meta = (DevelopmentOnly, AutoCreateRefTerm = "Location, Rotation, AxesSettings"))
+	UFUNCTION(BlueprintCallable, Category = "Draw Debug Library", meta = (DevelopmentOnly, AutoCreateRefTerm = "Location, Rotation, TextLineStyle, AxesLineStyle, AxesSettings"))
 	static UE_API void DrawDebugGraphAxes(
 		const FDebugDrawer& Drawer,
 		const FVector& Location,
 		const FRotator& Rotation,
 		const float XaxisLength = 100.0f,
 		const float YaxisLength = 100.0f,
-		const FLinearColor Color = FLinearColor::Red,
-		const float AxisThickness = 1.0f,
-		const float TextThickness = 0.0f,
+		const FDrawDebugLineStyle& TextLineStyle = FDrawDebugLineStyle(),
+		const FDrawDebugLineStyle& AxesLineStyle = FDrawDebugLineStyle(),
 		const bool bDepthTest = true,
 		const FDrawDebugGraphAxesSettings& AxesSettings = FDrawDebugGraphAxesSettings());
 
 	/** Debug Draw a line on a simple graph. */
-	UFUNCTION(BlueprintCallable, Category = "Draw Debug Library", meta = (DevelopmentOnly, AutoCreateRefTerm = "Location, Rotation, AxesSettings"))
+	UFUNCTION(BlueprintCallable, Category = "Draw Debug Library", meta = (DevelopmentOnly, AutoCreateRefTerm = "Location, Rotation, Xvalues, Yvalues, LineStyle"))
 	static UE_API void DrawDebugGraphLine(
 		const FDebugDrawer& Drawer,
 		const FVector& Location,
@@ -917,8 +1042,7 @@ public: // Graphs
 		const float Ymax = 1.0f,
 		const float XaxisLength = 100.0f,
 		const float YaxisLength = 100.0f,
-		const FLinearColor Color = FLinearColor::Red,
-		const float LineThickness = 0.0f,
+		const FDrawDebugLineStyle& LineStyle = FDrawDebugLineStyle(),
 		const bool bDepthTest = true);
 
 	static UE_API void DrawDebugGraphLineArrayView(
@@ -933,12 +1057,11 @@ public: // Graphs
 		const float Ymax = 1.0f,
 		const float XaxisLength = 100.0f,
 		const float YaxisLength = 100.0f,
-		const FLinearColor Color = FLinearColor::Red,
-		const float LineThickness = 0.0f,
+		const FDrawDebugLineStyle& LineStyle = FDrawDebugLineStyle(),
 		const bool bDepthTest = true);
 
 	/** Debug Draw a simple graph legend. */
-	UFUNCTION(BlueprintCallable, Category = "Draw Debug Library", meta = (DevelopmentOnly, AutoCreateRefTerm = "Location, Rotation, LegendSettings"))
+	UFUNCTION(BlueprintCallable, Category = "Draw Debug Library", meta = (DevelopmentOnly, AutoCreateRefTerm = "Location, Rotation, LegendColors, LegendLabels, TextLineStyle, IconLineStyle, LegendSettings"))
 	static UE_API void DrawDebugGraphLegend(
 		const FDebugDrawer& Drawer,
 		const FVector& Location,
@@ -947,8 +1070,9 @@ public: // Graphs
 		const TArray<FString>& LegendLabels,
 		const float XaxisLength = 100.0f,
 		const float YaxisLength = 100.0f,
-		const float TextThickness = 0.0f,
-		const float PointSize = 2.5f,
+		const FDrawDebugLineStyle& TextLineStyle = FDrawDebugLineStyle(),
+		const FDrawDebugLineStyle& IconLineStyle = FDrawDebugLineStyle(),
+		const float IconSize = 2.5f,
 		const bool bDepthTest = true,
 		const FDrawDebugStringSettings& LegendSettings = FDrawDebugStringSettings());
 
@@ -960,8 +1084,9 @@ public: // Graphs
 		const TArrayView<const FString> LegendLabels,
 		const float XaxisLength = 100.0f,
 		const float YaxisLength = 100.0f,
-		const float TextThickness = 0.0f,
-		const float PointSize = 2.5f,
+		const FDrawDebugLineStyle& TextLineStyle = FDrawDebugLineStyle(),
+		const FDrawDebugLineStyle& IconLineStyle = FDrawDebugLineStyle(),
+		const float IconSize = 2.5f,
 		const bool bDepthTest = true,
 		const FDrawDebugStringSettings& LegendSettings = FDrawDebugStringSettings());
 
