DrawDebugLibrary: Added more convience functions and drawing abilities

diff --git a/Engine/Plugins/Experimental/DrawDebugLibrary/DrawDebugLibrary.uplugin b/Engine/Plugins/Experimental/DrawDebugLibrary/DrawDebugLibrary.uplugin
--- a/Engine/Plugins/Experimental/DrawDebugLibrary/DrawDebugLibrary.uplugin
+++ b/Engine/Plugins/Experimental/DrawDebugLibrary/DrawDebugLibrary.uplugin
@@ -22,5 +22,10 @@
 			"LoadingPhase": "Default"
 		},
 	],
-	"Plugins": []
+	"Plugins": [
+		{
+			"Name": "RigVM",
+			"Enabled": true
+		}
+	]
 }
\ No newline at end of file
diff --git a/Engine/Plugins/Experimental/DrawDebugLibrary/Source/DrawDebugLibrary/DrawDebugLibrary.Build.cs b/Engine/Plugins/Experimental/DrawDebugLibrary/Source/DrawDebugLibrary/DrawDebugLibrary.Build.cs
--- a/Engine/Plugins/Experimental/DrawDebugLibrary/Source/DrawDebugLibrary/DrawDebugLibrary.Build.cs
+++ b/Engine/Plugins/Experimental/DrawDebugLibrary/Source/DrawDebugLibrary/DrawDebugLibrary.Build.cs
@@ -36,6 +36,7 @@ public class DrawDebugLibrary : ModuleRules
 			{
 				"CoreUObject",
 				"Engine",
+				"RigVM",
 				// ... add private dependencies that you statically link with here ...	
 			}
 			);
diff --git a/Engine/Plugins/Experimental/DrawDebugLibrary/Source/DrawDebugLibrary/Private/DrawDebugLibrary.cpp b/Engine/Plugins/Experimental/DrawDebugLibrary/Source/DrawDebugLibrary/Private/DrawDebugLibrary.cpp
--- a/Engine/Plugins/Experimental/DrawDebugLibrary/Source/DrawDebugLibrary/Private/DrawDebugLibrary.cpp
+++ b/Engine/Plugins/Experimental/DrawDebugLibrary/Source/DrawDebugLibrary/Private/DrawDebugLibrary.cpp
@@ -10,6 +10,7 @@
 #include "Animation/AnimInstanceProxy.h"
 #include "VisualLogger/VisualLogger.h"
 #include "Components/SkinnedMeshComponent.h"
+#include "RigVM/Public/RigVMCore/RigVMDrawInterface.h"
 
 void FDrawDebugLibraryModule::StartupModule()
 {
@@ -428,6 +429,11 @@ FDebugDrawer::FDebugDrawer(FPrimitiveDrawInterface* InPrimitiveDrawInterface) :
 	PDI = InPrimitiveDrawInterface;
 }
 
+FDebugDrawer::FDebugDrawer(FRigVMDrawInterface* InRigVMDrawInterface) : FDebugDrawer()
+{
+	RigVM = InRigVMDrawInterface;
+}
+
 FDebugDrawer FDebugDrawer::MakeDebugDrawer()
 {
 	return FDebugDrawer();
@@ -453,6 +459,11 @@ FDebugDrawer FDebugDrawer::MakeDebugDrawer(FPrimitiveDrawInterface* PrimitiveDra
 	return FDebugDrawer(PrimitiveDrawInterface);
 }
 
+FDebugDrawer FDebugDrawer::MakeDebugDrawer(FRigVMDrawInterface* InRigVMDrawInterface)
+{
+	return FDebugDrawer(InRigVMDrawInterface);
+}
+
 FDebugDrawer FDebugDrawer::MakeVisualLoggerDebugDrawer(
 	UObject* Object, 
 	const FName Category, 
@@ -492,6 +503,7 @@ FDebugDrawer FDebugDrawer::MakeMergedDebugDrawer(const TArrayView<const FDebugDr
 	for (const FDebugDrawer& SubDrawer : DebugDrawers)
 	{
 		if (SubDrawer.PDI) { Drawer.PDI = SubDrawer.PDI; }
+		if (SubDrawer.RigVM) { Drawer.RigVM = SubDrawer.RigVM; }
 		if (SubDrawer.World) { Drawer.World = SubDrawer.World; }
 		if (SubDrawer.AnimInstanceProxy) { Drawer.AnimInstanceProxy = SubDrawer.AnimInstanceProxy; }
 		if (SubDrawer.VisualLoggerObject)
@@ -519,6 +531,11 @@ void FDebugDrawer::DrawPoints(const TArrayView<const FVector> Locations, const F
 		}
 	}
 
+	if (RigVM)
+	{
+		RigVM->DrawPoints(FTransform::Identity, Locations, Thickness, Color, bDepthTest ? SDPG_World : SDPG_Foreground);
+	}
+
 #if UE_ENABLE_DEBUG_DRAWING
 	if (World)
 	{
@@ -610,19 +627,30 @@ void FDebugDrawer::DrawLines(const TArrayView<const TPair<FVector, FVector>> Seg
 
 	if (PDI)
 	{
-		for (int32 Idx = 0; Idx < Num; Idx++)
+		if (bDepthTest)
 		{
-			if (bDepthTest)
+			for (int32 Idx = 0; Idx < Num; Idx++)
 			{
 				PDI->DrawTranslucentLine(Segments[Idx].Key, Segments[Idx].Value, Color, SDPG_World, FMath::Max(Thickness, 0.0f), 0.0f, Thickness <= 0.0f);
 			}
-			else
+		}
+		else
+		{
+			for (int32 Idx = 0; Idx < Num; Idx++)
 			{
 				PDI->DrawTranslucentLine(Segments[Idx].Key, Segments[Idx].Value, Color, Thickness == 0.0f ? SDPG_Foreground : SDPG_World, Thickness, Thickness == 0.0f ? 1.0f : UE_MAX_FLT);
 			}
 		}
 	}
 
+	if (RigVM)
+	{
+		for (int32 Idx = 0; Idx < Num; Idx++)
+		{
+			RigVM->DrawLine(FTransform::Identity, Segments[Idx].Key, Segments[Idx].Value, Color, FMath::Max(Thickness, 0.0f), bDepthTest ? SDPG_World : SDPG_Foreground);
+		}
+	}
+
 #if UE_ENABLE_DEBUG_DRAWING
 	if (World)
 	{
@@ -837,6 +865,28 @@ void UDrawDebugLibrary::AddToVectorHistoryArray(TArray<FVector>& InOutValues, co
 	}
 }
 
+FDrawDebugPointStyle::FDrawDebugPointStyle(
+	const FLinearColor InColor,
+	const float InThickness) 
+	: Color(InColor)
+	, Thickness(InThickness) 
+{}
+
+FDrawDebugLineStyle::FDrawDebugLineStyle(
+	const FLinearColor InColor,
+	const float InThickness,
+	const EDrawDebugLineType InLineType,
+	const float InDashWidth,
+	const float InDashSpacing,
+	const float InDotSpacing)
+	: Color(InColor)
+	, Thickness(InThickness)
+	, LineType(InLineType)
+	, DashWidth(InDashWidth)
+	, DashSpacing(InDashSpacing)
+	, DotSpacing(InDotSpacing)
+{}
+
 FDrawDebugPointStyle UDrawDebugLibrary::MakeDrawDebugPointStyleFromColor(const FLinearColor Color)
 {
 	FDrawDebugPointStyle Out;
@@ -844,6 +894,14 @@ FDrawDebugPointStyle UDrawDebugLibrary::MakeDrawDebugPointStyleFromColor(const F
 	return Out;
 }
 
+FDrawDebugPointStyle UDrawDebugLibrary::MakeDrawDebugPointStyleFromColorAndThickness(const FLinearColor Color, const float Thickness)
+{
+	FDrawDebugPointStyle Out;
+	Out.Color = Color;
+	Out.Thickness = Thickness;
+	return Out;
+}
+
 FDrawDebugLineStyle UDrawDebugLibrary::MakeDrawDebugLineStyleFromColor(const FLinearColor Color)
 {
 	FDrawDebugLineStyle Out;
@@ -851,6 +909,14 @@ FDrawDebugLineStyle UDrawDebugLibrary::MakeDrawDebugLineStyleFromColor(const FLi
 	return Out;
 }
 
+FDrawDebugLineStyle UDrawDebugLibrary::MakeDrawDebugLineStyleFromColorAndThickness(const FLinearColor Color, const float Thickness)
+{
+	FDrawDebugLineStyle Out;
+	Out.Color = Color;
+	Out.Thickness = Thickness;
+	return Out;
+}
+
 FDrawDebugPointStyle UDrawDebugLibrary::DrawDebugPointStyleFromLineStyle(const FDrawDebugLineStyle& LineStyle)
 {
 	FDrawDebugPointStyle Out;
@@ -2015,6 +2081,74 @@ void UDrawDebugLibrary::DrawDebugCamera(const FDebugDrawer& Drawer, const FVecto
 		}, LineStyle, bDepthTest);
 }
 
+void UDrawDebugLibrary::DrawDebugBackpack(const FDebugDrawer& Drawer, const FVector& Location, const FRotator& Rotation, const FDrawDebugLineStyle& LineStyle, const bool bDepthTest, const FDrawDebugBackpackSettings& Settings)
+{
+	DrawDebugBox(
+		Drawer, 
+		Location + Rotation.RotateVector(FVector(Settings.Depth / 2, 0, 0)), 
+		Rotation, 
+		LineStyle, 
+		bDepthTest, 
+		FVector(Settings.Depth / 2, Settings.Width / 2, Settings.Height / 2));
+
+	const float W = Settings.Width;
+	const float H = Settings.Height;
+	const float SD = Settings.StrapDistance;
+	const float SW = Settings.StrapWidth;
+	const float SHI = Settings.StrapHorizontalInset;
+	const float SVI = Settings.StrapVerticalInset;
+	const float SCH = Settings.StrapContactHeight;
+
+	DrawDebugCatmullRomSplineArrayView(Drawer, {
+		Location + Rotation.RotateVector(FVector(0.0f, -W / 2 + SHI, +H / 2 - SVI)),
+		Location + Rotation.RotateVector(FVector( -SD, -W / 2 + SHI, -H / 2 + SCH)),
+		Location + Rotation.RotateVector(FVector(0.0f, -W / 2 + SHI, -H / 2 + SVI)),
+		}, LineStyle, bDepthTest);
+
+	DrawDebugCatmullRomSplineArrayView(Drawer, {
+		Location + Rotation.RotateVector(FVector(0.0f, -W / 2 + SHI + SW, +H / 2 - SVI)),
+		Location + Rotation.RotateVector(FVector( -SD, -W / 2 + SHI + SW, -H / 2 + SCH)),
+		Location + Rotation.RotateVector(FVector(0.0f, -W / 2 + SHI + SW, -H / 2 + SVI)),
+		}, LineStyle, bDepthTest);
+
+	DrawDebugCatmullRomSplineArrayView(Drawer, {
+		Location + Rotation.RotateVector(FVector(0.0f, +W / 2 - SHI, +H / 2 - SVI)),
+		Location + Rotation.RotateVector(FVector( -SD, +W / 2 - SHI, -H / 2 + SCH)),
+		Location + Rotation.RotateVector(FVector(0.0f, +W / 2 - SHI, -H / 2 + SVI)),
+		}, LineStyle, bDepthTest);
+
+	DrawDebugCatmullRomSplineArrayView(Drawer, {
+		Location + Rotation.RotateVector(FVector(0.0f, +W / 2 - SHI - SW, +H / 2 - SVI)),
+		Location + Rotation.RotateVector(FVector( -SD, +W / 2 - SHI - SW, -H / 2 + SCH)),
+		Location + Rotation.RotateVector(FVector(0.0f, +W / 2 - SHI - SW, -H / 2 + SVI)),
+		}, LineStyle, bDepthTest);
+}
+
+void UDrawDebugLibrary::DrawDebugBriefcase(const FDebugDrawer& Drawer, const FVector& Location, const FRotator& Rotation, const FDrawDebugLineStyle& LineStyle, const bool bDepthTest, const FDrawDebugBriefcaseSettings& Settings)
+{
+	DrawDebugBox(
+		Drawer,
+		Location + Rotation.RotateVector(FVector(0, 0, -Settings.HalfExtents.Z - Settings.HandleHeight)),
+		Rotation,
+		LineStyle,
+		bDepthTest,
+		Settings.HalfExtents);
+
+	DrawDebugCatmullRomSplineArrayView(Drawer, {
+		Location + Rotation.RotateVector(FVector(+Settings.HandleLength / 2, -Settings.HandleWidth / 2, -Settings.HandleHeight)),
+		Location + Rotation.RotateVector(FVector(+Settings.HandleLength / 2 - Settings.HandleLength / 8, -Settings.HandleWidth / 2, 0.0f)),
+		Location + Rotation.RotateVector(FVector(-Settings.HandleLength / 2 + Settings.HandleLength / 8, -Settings.HandleWidth / 2, 0.0f)),
+		Location + Rotation.RotateVector(FVector(-Settings.HandleLength / 2, -Settings.HandleWidth / 2, -Settings.HandleHeight)),
+		}, LineStyle, bDepthTest, true);
+
+	DrawDebugCatmullRomSplineArrayView(Drawer, {
+		Location + Rotation.RotateVector(FVector(+Settings.HandleLength / 2, +Settings.HandleWidth / 2, -Settings.HandleHeight)),
+		Location + Rotation.RotateVector(FVector(+Settings.HandleLength / 2 - Settings.HandleLength / 8, +Settings.HandleWidth / 2, 0.0f)),
+		Location + Rotation.RotateVector(FVector(-Settings.HandleLength / 2 + Settings.HandleLength / 8, +Settings.HandleWidth / 2, 0.0f)),
+		Location + Rotation.RotateVector(FVector(-Settings.HandleLength / 2, +Settings.HandleWidth / 2, -Settings.HandleHeight)),
+		}, LineStyle, bDepthTest, true);
+}
+
 void UDrawDebugLibrary::DrawDebugTrajectory(const FDebugDrawer& Drawer, const TArray<FVector>& Locations, const TArray<FVector>& Directions, const FTransform& RelativeTransform, const FDrawDebugLineStyle& LineStyle, const bool bDepthTest, const float DrawArrowLength, const float PointRadius, const float ArrowHeadScale, const int32 Segments, const float VerticalOffset)
 {
 	DrawDebugTrajectoryFromArrayViews(Drawer, Locations, Directions, RelativeTransform, LineStyle, bDepthTest, DrawArrowLength, PointRadius, ArrowHeadScale, Segments, VerticalOffset);
diff --git a/Engine/Plugins/Experimental/DrawDebugLibrary/Source/DrawDebugLibrary/Public/DrawDebugLibrary.h b/Engine/Plugins/Experimental/DrawDebugLibrary/Source/DrawDebugLibrary/Public/DrawDebugLibrary.h
--- a/Engine/Plugins/Experimental/DrawDebugLibrary/Source/DrawDebugLibrary/Public/DrawDebugLibrary.h
+++ b/Engine/Plugins/Experimental/DrawDebugLibrary/Source/DrawDebugLibrary/Public/DrawDebugLibrary.h
@@ -11,6 +11,7 @@
 #define UE_API DRAWDEBUGLIBRARY_API
 
 class FPrimitiveDrawInterface;
+struct FRigVMDrawInterface;
 struct FTransformTrajectory;
 struct FAnimInstanceProxy;
 class USkinnedMeshComponent;
@@ -69,13 +70,18 @@ struct FDrawDebugPointStyle
 {
 	GENERATED_BODY()
 
-	/** Thickness of the point. Set to 0.0 for pixel-sized points. */
-	UPROPERTY(BlueprintReadWrite, EditAnywhere, Category = "Draw Debug Library", meta = (ClampMin = "0", UIMin = "0", Units = "cm"))
-	float Thickness = 0.0f;
+	/** Construct a Point Style */
+	UE_API FDrawDebugPointStyle(
+		const FLinearColor InColor = FLinearColor(1.0f, 0.0f, 1.0f, 1.0f),
+		const float InThickness = 0.0f);
 
 	/** Color of the point. Alpha controls the opacity. */
 	UPROPERTY(BlueprintReadWrite, EditAnywhere, Category = "Draw Debug Library")
 	FLinearColor Color = FLinearColor(1.0f, 0.0f, 1.0f, 1.0f);
+
+	/** Thickness of the point. Set to 0.0 for pixel-sized points. */
+	UPROPERTY(BlueprintReadWrite, EditAnywhere, Category = "Draw Debug Library", meta = (ClampMin = "0", UIMin = "0", Units = "cm"))
+	float Thickness = 0.0f;
 };
 
 
@@ -100,28 +106,37 @@ struct FDrawDebugLineStyle
 {
 	GENERATED_BODY()
 
-	/** Line style to draw with */
+	/** Construct a Line Style */
+	UE_API FDrawDebugLineStyle(
+		const FLinearColor InColor = FLinearColor(1.0f, 0.0f, 1.0f, 1.0f),
+		const float InThickness = 0.0f,
+		const EDrawDebugLineType InLineType = EDrawDebugLineType::Solid,
+		const float InDashWidth = 2.5f,
+		const float InDashSpacing = 1.0f,
+		const float InDotSpacing = 1.0f);
+
+	/** Color of the line. Alpha controls the opacity. */
 	UPROPERTY(BlueprintReadWrite, EditAnywhere, Category = "Draw Debug Library")
-	EDrawDebugLineType LineType = EDrawDebugLineType::Solid;
+	FLinearColor Color = FLinearColor(1.0f, 0.0f, 1.0f, 1.0f);
 
 	/** Thickness of the line. Set to 0.0 for pixel-sized lines. */
 	UPROPERTY(BlueprintReadWrite, EditAnywhere, Category = "Draw Debug Library", meta = (ClampMin = "0", UIMin = "0", Units = "cm"))
 	float Thickness = 0.0f;
 
-	/** Color of the line. Alpha controls the opacity. */
+	/** Line style to draw with */
 	UPROPERTY(BlueprintReadWrite, EditAnywhere, Category = "Draw Debug Library")
-	FLinearColor Color = FLinearColor(1.0f, 0.0f, 1.0f, 1.0f);
+	EDrawDebugLineType LineType = EDrawDebugLineType::Solid;
 
 	/** Width of the dashes when using dashed style. */
-	UPROPERTY(BlueprintReadWrite, EditAnywhere, Category = "Draw Debug Library", meta = (ClampMin = "0", UIMin = "0", Units = "cm"))
+	UPROPERTY(BlueprintReadWrite, EditAnywhere, Category = "Draw Debug Library", AdvancedDisplay, meta = (ClampMin = "0", UIMin = "0", Units = "cm"))
 	float DashWidth = 2.5f;
 
 	/** Spacing of the dashes when using dashed style. */
-	UPROPERTY(BlueprintReadWrite, EditAnywhere, Category = "Draw Debug Library", meta = (ClampMin = "0", UIMin = "0", Units = "cm"))
+	UPROPERTY(BlueprintReadWrite, EditAnywhere, Category = "Draw Debug Library", AdvancedDisplay, meta = (ClampMin = "0", UIMin = "0", Units = "cm"))
 	float DashSpacing = 1.0f;
 
 	/** Spacing of the dots when using dotted style. */
-	UPROPERTY(BlueprintReadWrite, EditAnywhere, Category = "Draw Debug Library", meta = (ClampMin = "0", UIMin = "0", Units = "cm"))
+	UPROPERTY(BlueprintReadWrite, EditAnywhere, Category = "Draw Debug Library", AdvancedDisplay, meta = (ClampMin = "0", UIMin = "0", Units = "cm"))
 	float DotSpacing = 1.0f;
 };
 
@@ -227,6 +242,68 @@ struct FDrawDebugChairSettings
 	float BackTilt = 10.0f;
 };
 
+/** Settings for debug drawing backpacks */
+USTRUCT(BlueprintType, Category = "Draw Debug Library")
+struct FDrawDebugBackpackSettings
+{
+	GENERATED_BODY()
+
+	/** Width of the backpack */
+	UPROPERTY(BlueprintReadWrite, EditAnywhere, Category = "Draw Debug Library", meta = (ClampMin = "0", UIMin = "0", Units = "cm"))
+	float Width = 30.0f;
+
+	/** Height of the backpack */
+	UPROPERTY(BlueprintReadWrite, EditAnywhere, Category = "Draw Debug Library", meta = (ClampMin = "0", UIMin = "0", Units = "cm"))
+	float Height = 45.0f;
+
+	/** Depth of the backpack */
+	UPROPERTY(BlueprintReadWrite, EditAnywhere, Category = "Draw Debug Library", meta = (ClampMin = "0", UIMin = "0", Units = "cm"))
+	float Depth = 15.0f;
+
+	/** Distance the straps extend outward */
+	UPROPERTY(BlueprintReadWrite, EditAnywhere, Category = "Draw Debug Library", meta = (ClampMin = "0", UIMin = "0", Units = "cm"))
+	float StrapDistance = 25.0f;
+
+	/** The height at which the straps are in contact */
+	UPROPERTY(BlueprintReadWrite, EditAnywhere, Category = "Draw Debug Library", meta = (Units = "cm"))
+	float StrapContactHeight = 25.0f;
+
+	/** Width of the straps */
+	UPROPERTY(BlueprintReadWrite, EditAnywhere, Category = "Draw Debug Library", meta = (Units = "cm"))
+	float StrapWidth = 2.5f;
+
+	/** Inset of the straps horizontally */
+	UPROPERTY(BlueprintReadWrite, EditAnywhere, Category = "Draw Debug Library", meta = (Units = "cm"))
+	float StrapHorizontalInset = 5.0f;
+
+	/** Inset of the straps vertically */
+	UPROPERTY(BlueprintReadWrite, EditAnywhere, Category = "Draw Debug Library", meta = (Units = "cm"))
+	float StrapVerticalInset = 2.5f;
+};
+
+/** Settings for debug drawing briefcases */
+USTRUCT(BlueprintType, Category = "Draw Debug Library")
+struct FDrawDebugBriefcaseSettings
+{
+	GENERATED_BODY()
+
+	/** Half-extents of the briefcase */
+	UPROPERTY(BlueprintReadWrite, EditAnywhere, Category = "Draw Debug Library", meta = (ClampMin = "0", UIMin = "0", Units = "cm"))
+	FVector HalfExtents = FVector(20.0f, 5.0f, 15.0f);
+
+	/** Length of the handle */
+	UPROPERTY(BlueprintReadWrite, EditAnywhere, Category = "Draw Debug Library", meta = (ClampMin = "0", UIMin = "0", Units = "cm"))
+	float HandleLength = 10.0f;
+
+	/** The height of the handle above the briefcase */
+	UPROPERTY(BlueprintReadWrite, EditAnywhere, Category = "Draw Debug Library", meta = (ClampMin = "0", UIMin = "0", Units = "cm"))
+	float HandleHeight = 5.0f;
+
+	/** Width of the handle */
+	UPROPERTY(BlueprintReadWrite, EditAnywhere, Category = "Draw Debug Library", meta = (ClampMin = "0", UIMin = "0", Units = "cm"))
+	float HandleWidth = 2.5f;
+};
+
 /** Settings for debug drawing doors */
 USTRUCT(BlueprintType, Category = "Draw Debug Library")
 struct FDrawDebugDoorSettings
@@ -398,6 +475,9 @@ public: // Constructors
 	/** Make a debug drawer for a FPrimitiveDrawInterface */
 	UE_API FDebugDrawer(FPrimitiveDrawInterface* InPrimitiveDrawInterface);
 
+	/** Make a debug drawer for a FRigVMDrawInterface */
+	UE_API FDebugDrawer(FRigVMDrawInterface* InRigVMDrawInterface);
+
 	/** Make a null debug drawer which ignores draw commands */
 	static UE_API FDebugDrawer MakeDebugDrawer();
 
@@ -413,6 +493,9 @@ public: // Constructors
 	/** Make a debug drawer for a FPrimitiveDrawInterface */
 	static UE_API FDebugDrawer MakeDebugDrawer(FPrimitiveDrawInterface* PrimitiveDrawInterface);
 
+	/** Make a debug drawer for a FRigVMDrawInterface */
+	static UE_API FDebugDrawer MakeDebugDrawer(FRigVMDrawInterface* InRigVMDrawInterface);
+
 	/** Make a debug drawer for the Visual Logger */
 	static UE_API FDebugDrawer MakeVisualLoggerDebugDrawer(
 		UObject* Object, 
@@ -445,6 +528,9 @@ private:
 	/** Primitive Draw Interface - used for drawing in editor viewports. */
 	FPrimitiveDrawInterface* PDI = nullptr;
 
+	/** RigVM Draw Interface - used for drawing in a RigVM context. */
+	FRigVMDrawInterface* RigVM = nullptr;
+
 	/** World object - used for normal debug drawing. */
 	UPROPERTY()
 	TObjectPtr<UWorld> World = nullptr;
@@ -567,10 +653,18 @@ public: // Style Modifiers
 	UFUNCTION(BlueprintPure, Category = "Draw Debug Library", meta = (BlueprintAutocast, CompactNodeTitle = "->"))
 	static UE_API FDrawDebugPointStyle MakeDrawDebugPointStyleFromColor(const FLinearColor Color = FLinearColor(1.0f, 0.0f, 1.0f, 1.0f));
 
+	/** Make a point style from a color and thickness */
+	UFUNCTION(BlueprintPure, Category = "Draw Debug Library")
+	static UE_API FDrawDebugPointStyle MakeDrawDebugPointStyleFromColorAndThickness(const FLinearColor Color = FLinearColor(1.0f, 0.0f, 1.0f, 1.0f), const float Thickness = 0.0f);
+
 	/** Make a line style from a color */
 	UFUNCTION(BlueprintPure, Category = "Draw Debug Library", meta = (BlueprintAutocast, CompactNodeTitle = "->"))
 	static UE_API FDrawDebugLineStyle MakeDrawDebugLineStyleFromColor(const FLinearColor Color = FLinearColor(1.0f, 0.0f, 1.0f, 1.0f));
 
+	/** Make a line style from a color and thickness */
+	UFUNCTION(BlueprintPure, Category = "Draw Debug Library")
+	static UE_API FDrawDebugLineStyle MakeDrawDebugLineStyleFromColorAndThickness(const FLinearColor Color = FLinearColor(1.0f, 0.0f, 1.0f, 1.0f), const float Thickness = 0.0f);
+
 	/** Get a point style from a line style */
 	UFUNCTION(BlueprintPure, Category = "Draw Debug Library", meta = (BlueprintAutocast, CompactNodeTitle = "->", AutoCreateRefTerm = "LineStyle"))
 	static UE_API FDrawDebugPointStyle DrawDebugPointStyleFromLineStyle(const FDrawDebugLineStyle& LineStyle);
@@ -847,6 +941,14 @@ public: // Complex Objects
 	UFUNCTION(BlueprintCallable, Category = "Draw Debug Library", meta = (DevelopmentOnly, AutoCreateRefTerm = "Location, Rotation, LineStyle"))
 	static UE_API void DrawDebugCamera(const FDebugDrawer& Drawer, const FVector& Location, const FRotator& Rotation, const FDrawDebugLineStyle& LineStyle = FDrawDebugLineStyle(), const bool bDepthTest = true, const float Scale = 10.0f, const float FOVDegees = 30.0f);
 
+	/** Debug draw a backpack */
+	UFUNCTION(BlueprintCallable, Category = "Draw Debug Library", meta = (DevelopmentOnly, AutoCreateRefTerm = "Location, Rotation, LineStyle, Settings"))
+	static UE_API void DrawDebugBackpack(const FDebugDrawer& Drawer, const FVector& Location, const FRotator& Rotation, const FDrawDebugLineStyle& LineStyle = FDrawDebugLineStyle(), const bool bDepthTest = true, const FDrawDebugBackpackSettings& Settings = FDrawDebugBackpackSettings());
+
+	/** Debug draw a briefcase */
+	UFUNCTION(BlueprintCallable, Category = "Draw Debug Library", meta = (DevelopmentOnly, AutoCreateRefTerm = "Location, Rotation, LineStyle, Settings"))
+	static UE_API void DrawDebugBriefcase(const FDebugDrawer& Drawer, const FVector& Location, const FRotator& Rotation, const FDrawDebugLineStyle& LineStyle = FDrawDebugLineStyle(), const bool bDepthTest = true, const FDrawDebugBriefcaseSettings& Settings = FDrawDebugBriefcaseSettings());
+
 	/** Debug Draw a trajectory of locations and directions */
 	UFUNCTION(BlueprintCallable, Category = "Draw Debug Library", meta = (DevelopmentOnly, AutoCreateRefTerm = "Locations, Directions, RelativeTransform, LineStyle"))
 	static UE_API void DrawDebugTrajectory(const FDebugDrawer& Drawer, const TArray<FVector>& Locations, const TArray<FVector>& Directions, const FTransform& RelativeTransform = FTransform(), const FDrawDebugLineStyle& LineStyle = FDrawDebugLineStyle(), const bool bDepthTest = true, const float DrawArrowLength = 100.0f, const float PointRadius = 5.0f, const float ArrowHeadScale = 1.0f, const int32 Segments = 9, const float VerticalOffset = 2.0f);
