DrawDebugLibrary: Added no arrow head and support for specifying the arrow options for the axes when drawing a graph. Made functions BlueprintThreadsafe.

diff --git a/Engine/Plugins/Experimental/DrawDebugLibrary/Source/DrawDebugLibrary/Private/DrawDebugLibrary.cpp b/Engine/Plugins/Experimental/DrawDebugLibrary/Source/DrawDebugLibrary/Private/DrawDebugLibrary.cpp
--- a/Engine/Plugins/Experimental/DrawDebugLibrary/Source/DrawDebugLibrary/Private/DrawDebugLibrary.cpp
+++ b/Engine/Plugins/Experimental/DrawDebugLibrary/Source/DrawDebugLibrary/Private/DrawDebugLibrary.cpp
@@ -523,7 +523,7 @@ void FDebugDrawer::DrawPoints(const TArrayView<const FVector> Locations, const F
 {
 	const int32 Num = Locations.Num();
 
-	if (PDI)
+	if (PDI && IsInGameThread())
 	{
 		for (int32 Idx = 0; Idx < Num; Idx++)
 		{
@@ -537,7 +537,7 @@ void FDebugDrawer::DrawPoints(const TArrayView<const FVector> Locations, const F
 	}
 
 #if UE_ENABLE_DEBUG_DRAWING
-	if (World)
+	if (World && IsInGameThread())
 	{
 		if (Thickness == 0.0f)
 		{
@@ -600,7 +600,7 @@ void FDebugDrawer::DrawPoints(const TArrayView<const FVector> Locations, const F
 		
 		const UWorld* VisualLoggerWorld = VisualLoggerObject->GetWorld();
 
-		if (VisualLoggerWorld && bVisualLoggerDrawToScene && (bVisualLoggerDrawToSceneWhileRecording || !FVisualLogger::IsRecording()))
+		if (VisualLoggerWorld && IsInGameThread() && bVisualLoggerDrawToScene && (bVisualLoggerDrawToSceneWhileRecording || !FVisualLogger::IsRecording()))
 		{
 			if (Thickness == 0.0f)
 			{
@@ -625,7 +625,7 @@ void FDebugDrawer::DrawLines(const TArrayView<const TPair<FVector, FVector>> Seg
 {
 	const int32 Num = Segments.Num();
 
-	if (PDI)
+	if (PDI && IsInGameThread())
 	{
 		if (bDepthTest)
 		{
@@ -652,7 +652,7 @@ void FDebugDrawer::DrawLines(const TArrayView<const TPair<FVector, FVector>> Seg
 	}
 
 #if UE_ENABLE_DEBUG_DRAWING
-	if (World)
+	if (World && IsInGameThread())
 	{
 		for (int32 Idx = 0; Idx < Num; Idx++)
 		{
@@ -690,7 +690,7 @@ void FDebugDrawer::DrawLines(const TArrayView<const TPair<FVector, FVector>> Seg
 
 		const UWorld* VisualLoggerWorld = VisualLoggerObject->GetWorld();
 
-		if (VisualLoggerWorld && bVisualLoggerDrawToScene && (bVisualLoggerDrawToSceneWhileRecording || !FVisualLogger::IsRecording()))
+		if (VisualLoggerWorld && IsInGameThread() && bVisualLoggerDrawToScene && (bVisualLoggerDrawToSceneWhileRecording || !FVisualLogger::IsRecording()))
 		{
 			for (int32 Idx = 0; Idx < Num; Idx++)
 			{
@@ -887,7 +887,7 @@ FDrawDebugLineStyle::FDrawDebugLineStyle(
 	, DotSpacing(InDotSpacing)
 {}
 
-FDrawDebugPointStyle UDrawDebugLibrary::MakeDrawDebugPointStyleFromColor(const FLinearColor Color)
+FDrawDebugPointStyle UDrawDebugLibrary::MakeDrawDebugPointStyleFromColor(const FLinearColor& Color)
 {
 	FDrawDebugPointStyle Out;
 	Out.Color = Color;
@@ -902,7 +902,7 @@ FDrawDebugPointStyle UDrawDebugLibrary::MakeDrawDebugPointStyleFromColorAndThick
 	return Out;
 }
 
-FDrawDebugLineStyle UDrawDebugLibrary::MakeDrawDebugLineStyleFromColor(const FLinearColor Color)
+FDrawDebugLineStyle UDrawDebugLibrary::MakeDrawDebugLineStyleFromColor(const FLinearColor& Color)
 {
 	FDrawDebugLineStyle Out;
 	Out.Color = Color;
@@ -1494,6 +1494,8 @@ void UDrawDebugLibrary::DrawDebugOrientedArrow(
 	{
 		switch (Settings.ArrowHeadStartType)
 		{
+		case EDrawDebugArrowHead::None: { break; }
+
 		case EDrawDebugArrowHead::Simple:
 		{
 			DrawDebugLinesPairsArrayView(Drawer, {
@@ -1591,6 +1593,8 @@ void UDrawDebugLibrary::DrawDebugOrientedArrow(
 	{
 		switch (Settings.ArrowHeadEndType)
 		{
+		case EDrawDebugArrowHead::None: { break; }
+
 		case EDrawDebugArrowHead::Simple:
 		{
 			DrawDebugLinesPairsArrayView(Drawer, {
@@ -2874,7 +2878,8 @@ void UDrawDebugLibrary::DrawDebugGraphAxes(
 		(Rotation.Quaternion() * FQuat::MakeFromRotationVector(FVector(UE_PI / 2, 0, 0))).Rotator(),
 		XaxisLength,
 		AxesLineStyle,
-		bDepthTest);
+		bDepthTest,
+		AxesSettings.AxisArrowSettings);
 
 	DrawDebugOrientedArrow(
 		Drawer,
@@ -2884,7 +2889,8 @@ void UDrawDebugLibrary::DrawDebugGraphAxes(
 			FQuat::MakeFromRotationVector(FVector(UE_PI / 2, 0, 0))).Rotator(),
 		YaxisLength,
 		AxesLineStyle,
-		bDepthTest);
+		bDepthTest,
+		AxesSettings.AxisArrowSettings);
 
 	const FVector TitleDimension = DrawDebugStringViewDimensions(AxesSettings.Title, AxesSettings.TitleSettings);
 
diff --git a/Engine/Plugins/Experimental/DrawDebugLibrary/Source/DrawDebugLibrary/Public/DrawDebugLibrary.h b/Engine/Plugins/Experimental/DrawDebugLibrary/Source/DrawDebugLibrary/Public/DrawDebugLibrary.h
--- a/Engine/Plugins/Experimental/DrawDebugLibrary/Source/DrawDebugLibrary/Public/DrawDebugLibrary.h
+++ b/Engine/Plugins/Experimental/DrawDebugLibrary/Source/DrawDebugLibrary/Public/DrawDebugLibrary.h
@@ -177,6 +177,9 @@ enum class EDrawDebugArrowHead : uint8
 
 	/** 3D box arrow head */
 	Box,
+
+	/** No arrow head */
+	None,
 };
 
 
@@ -443,6 +446,10 @@ struct FDrawDebugGraphAxesSettings
 	UPROPERTY(BlueprintReadWrite, EditAnywhere, Category = "Draw Debug Library")
 	FString YaxisLabel;
 
+	/** Settings for the axis arrows */
+	UPROPERTY(BlueprintReadWrite, EditAnywhere, Category = "Draw Debug Library")
+	FDrawDebugArrowSettings AxisArrowSettings;
+
 	/** Settings for the axis labels */
 	UPROPERTY(BlueprintReadWrite, EditAnywhere, Category = "Draw Debug Library")
 	FDrawDebugStringSettings AxisLabelSettings;
@@ -555,7 +562,6 @@ private:
 	bool bVisualLoggerDrawToSceneWhileRecording = true;
 };
 
-
 /**
  * A blueprint library of additional debug draw functions that unifies drawing between different debug drawing interfaces. Also includes a number of 
  * helper functions and other convenience functions that are useful for debug drawing such as functions for recording histories.
@@ -564,7 +570,7 @@ private:
  * 
  *	Required Inputs, Style, bDepthTest, Settings
  */
-UCLASS(BlueprintType, Category = "Draw Debug Library")
+UCLASS(BlueprintType, Category = "Draw Debug Library", meta = (BlueprintThreadSafe))
 class UDrawDebugLibrary : public UBlueprintFunctionLibrary
 {
 	GENERATED_BODY()
@@ -636,30 +642,30 @@ public: // Visual Debugger Logging
 public: // Convenience functions
 
 	/** Convenience function for making an array of linearly spaced float values */
-	UFUNCTION(BlueprintPure = true, Category = "Draw Debug Library")
+	UFUNCTION(BlueprintPure = true, Category = "Draw Debug Library", meta = (NotBlueprintThreadSafe))
 	static UE_API void MakeLinearlySpacedFloatArray(TArray<float>& OutValues, const float Start = 0.0f, const float Stop = 1.0f, const int32 Num = 10);
 
 	/** Convenience function that adds a float to the end of an array, popping values from the front once the max is reached */
-	UFUNCTION(BlueprintPure = false, Category = "Draw Debug Library")
+	UFUNCTION(BlueprintPure = false, Category = "Draw Debug Library", meta = (NotBlueprintThreadSafe))
 	static UE_API void AddToFloatHistoryArray(UPARAM(ref) TArray<float>& InOutValues, const float NewValue = 0.0f, const int32 MaxHistoryNum = 60);
 
 	/** Convenience function that adds a vector to the end of an array, popping values from the front once the max is reached */
-	UFUNCTION(BlueprintPure = false, Category = "Draw Debug Library")
+	UFUNCTION(BlueprintPure = false, Category = "Draw Debug Library", meta = (NotBlueprintThreadSafe))
 	static UE_API void AddToVectorHistoryArray(UPARAM(ref) TArray<FVector>& InOutValues, const FVector NewValue = FVector::ZeroVector, const int32 MaxHistoryNum = 60);
 
 public: // Style Modifiers
 
 	/** Make a point style from a color */
-	UFUNCTION(BlueprintPure, Category = "Draw Debug Library", meta = (BlueprintAutocast, CompactNodeTitle = "->"))
-	static UE_API FDrawDebugPointStyle MakeDrawDebugPointStyleFromColor(const FLinearColor Color = FLinearColor(1.0f, 0.0f, 1.0f, 1.0f));
+	UFUNCTION(BlueprintPure, Category = "Draw Debug Library", meta = (BlueprintAutocast, CompactNodeTitle = "->", AutoCreateRefTerm = "Color"))
+	static UE_API FDrawDebugPointStyle MakeDrawDebugPointStyleFromColor(const FLinearColor& Color);
 
 	/** Make a point style from a color and thickness */
 	UFUNCTION(BlueprintPure, Category = "Draw Debug Library")
 	static UE_API FDrawDebugPointStyle MakeDrawDebugPointStyleFromColorAndThickness(const FLinearColor Color = FLinearColor(1.0f, 0.0f, 1.0f, 1.0f), const float Thickness = 0.0f);
 
 	/** Make a line style from a color */
-	UFUNCTION(BlueprintPure, Category = "Draw Debug Library", meta = (BlueprintAutocast, CompactNodeTitle = "->"))
-	static UE_API FDrawDebugLineStyle MakeDrawDebugLineStyleFromColor(const FLinearColor Color = FLinearColor(1.0f, 0.0f, 1.0f, 1.0f));
+	UFUNCTION(BlueprintPure, Category = "Draw Debug Library", meta = (BlueprintAutocast, CompactNodeTitle = "->", AutoCreateRefTerm = "Color"))
+	static UE_API FDrawDebugLineStyle MakeDrawDebugLineStyleFromColor(const FLinearColor& Color);
 
 	/** Make a line style from a color and thickness */
 	UFUNCTION(BlueprintPure, Category = "Draw Debug Library")
