DrawDebugLibrary: Switched text to drawing on YZ plane.

#virtualized
#Synced-CL 48146366

diff --git a/Engine/Plugins/Experimental/DrawDebugLibrary/Source/DrawDebugLibrary/Private/DrawDebugLibrary.cpp b/Engine/Plugins/Experimental/DrawDebugLibrary/Source/DrawDebugLibrary/Private/DrawDebugLibrary.cpp
--- a/Engine/Plugins/Experimental/DrawDebugLibrary/Source/DrawDebugLibrary/Private/DrawDebugLibrary.cpp
+++ b/Engine/Plugins/Experimental/DrawDebugLibrary/Source/DrawDebugLibrary/Private/DrawDebugLibrary.cpp
@@ -860,18 +860,21 @@ void UDrawDebugLibrary::AddToFloatHistoryArray(TArray<float>& InOutValues, const
 		return;
 	}
 
-	InOutValues.Add(NewValue);
+	const int32 Shift = InOutValues.Num() - MaxHistoryNum + 1;
 
-	while (InOutValues.Num() > MaxHistoryNum)
+	if (Shift > 0)
 	{
-		// Shift Values
-		for (int32 Idx = 0; Idx < InOutValues.Num() - 1; Idx++)
+		for (int32 Idx = 0; Idx < InOutValues.Num() - Shift; Idx++)
 		{
-			InOutValues[Idx] = InOutValues[Idx + 1];
+			InOutValues[Idx] = InOutValues[Idx + Shift];
 		}
 
-		// Pop End
-		InOutValues.Pop();
+		InOutValues.SetNum(MaxHistoryNum);
+		InOutValues.Last() = NewValue;
+	}
+	else
+	{
+		InOutValues.Add(NewValue);
 	}
 }
 
@@ -883,21 +886,113 @@ void UDrawDebugLibrary::AddToVectorHistoryArray(TArray<FVector>& InOutValues, co
 		return;
 	}
 
-	InOutValues.Add(NewValue);
+	const int32 Shift = InOutValues.Num() - MaxHistoryNum + 1;
+
+	if (Shift > 0)
+	{
+		for (int32 Idx = 0; Idx < InOutValues.Num() - Shift; Idx++)
+		{
+			InOutValues[Idx] = InOutValues[Idx + Shift];
+		}
+
+		InOutValues.SetNum(MaxHistoryNum);
+		InOutValues.Last() = NewValue;
+	}
+	else
+	{
+		InOutValues.Add(NewValue);
+	}
+}
+
+void UDrawDebugLibrary::AddToNameHistoryArray(TArray<FName>& InOutValues, const FName NewValue, const int32 MaxHistoryNum)
+{
+	if (MaxHistoryNum <= 0)
+	{
+		InOutValues.Empty();
+		return;
+	}
+
+	const int32 Shift = InOutValues.Num() - MaxHistoryNum + 1;
 
-	while (InOutValues.Num() > MaxHistoryNum)
+	if (Shift > 0)
 	{
-		// Shift Values
-		for (int32 Idx = 0; Idx < InOutValues.Num() - 1; Idx++)
+		for (int32 Idx = 0; Idx < InOutValues.Num() - Shift; Idx++)
 		{
-			InOutValues[Idx] = InOutValues[Idx + 1];
+			InOutValues[Idx] = InOutValues[Idx + Shift];
 		}
 
-		// Pop End
-		InOutValues.Pop();
+		InOutValues.SetNum(MaxHistoryNum);
+		InOutValues.Last() = NewValue;
+	}
+	else
+	{
+		InOutValues.Add(NewValue);
 	}
 }
 
+void UDrawDebugLibrary::DrawDebugLocalOffset(
+	FVector& OutLocation,
+	FRotator& OutRotation,
+	const FVector DrawLocation,
+	const FRotator DrawRotation,
+	const FVector DrawOffset)
+{
+	OutLocation = DrawLocation + DrawRotation.RotateVector(DrawOffset);
+	OutRotation = DrawRotation;
+}
+
+void UDrawDebugLibrary::DrawDebugOrientUprightToFloor(
+	FVector& OutLocation,
+	FRotator& OutRotation,
+	const FVector DrawWorldLocation,
+	const FVector DrawUprightOffset)
+{
+	const FQuat Reorientation = FQuat::MakeFromRotationVector(FVector(0.0f, +UE_PI / 2, 0.0f));
+	OutLocation = DrawWorldLocation + Reorientation.RotateVector(DrawUprightOffset);
+	OutRotation = Reorientation.Rotator();
+}
+
+void UDrawDebugLibrary::DrawDebugOrientFloorToUpright(
+	FVector& OutLocation,
+	FRotator& OutRotation,
+	const FVector DrawWorldLocation,
+	const FVector DrawFloorOffset)
+{
+	const FQuat Reorientation = FQuat::MakeFromRotationVector(FVector(0.0f, -UE_PI / 2, 0.0f));
+	OutLocation = DrawWorldLocation + Reorientation.RotateVector(DrawFloorOffset);
+	OutRotation = Reorientation.Rotator();
+}
+
+void UDrawDebugLibrary::DrawDebugOrientUprightToCamera(
+	FVector& OutLocation,
+	FRotator& OutRotation,
+	const FVector DrawWorldLocation,
+	const FVector DrawUprightOffset,
+	const FRotator CameraRotation,
+	const bool bUseOnlyYaw)
+{
+	const FRotator ActualCameraRotation = bUseOnlyYaw ? FRotator(0.0f, CameraRotation.Yaw, 0.0f) : CameraRotation;
+
+	OutLocation = DrawWorldLocation + ActualCameraRotation.RotateVector(DrawUprightOffset);
+	OutRotation = ActualCameraRotation;
+}
+
+void UDrawDebugLibrary::DrawDebugOrientFloorToCamera(
+	FVector& OutLocation,
+	FRotator& OutRotation,
+	const FVector DrawWorldLocation,
+	const FVector DrawFloorOffset,
+	const FRotator CameraRotation,
+	const bool bUseOnlyYaw)
+{
+	const FQuat Reorientation = FQuat::MakeFromRotationVector(FVector(0.0f, -UE_PI / 2, 0.0f));
+	const FRotator ActualCameraRotation = bUseOnlyYaw ? FRotator(0.0f, CameraRotation.Yaw, 0.0f) : CameraRotation;
+	const FQuat FinalRotation = ActualCameraRotation.Quaternion() * Reorientation;
+
+	OutLocation = DrawWorldLocation + FinalRotation.RotateVector(DrawFloorOffset);
+	OutRotation = FinalRotation.Rotator();
+}
+
 FDrawDebugPointStyle::FDrawDebugPointStyle(
 	const FLinearColor InColor,
 	const float InThickness) 
@@ -1240,6 +1335,19 @@ void UDrawDebugLibrary::DrawDebugSquare(const FDebugDrawer& Drawer, const FVecto
 		}, LineStyle, bDepthTest);
 }
 
+void UDrawDebugLibrary::DrawDebugCross(const FDebugDrawer& Drawer, const FVector& Location, const FRotator& Rotation, const FDrawDebugLineStyle& LineStyle, const bool bDepthTest, const float HalfLength)
+{
+	const FVector V0 = Location + Rotation.RotateVector(HalfLength * FVector(+1, +1, 0));
+	const FVector V1 = Location + Rotation.RotateVector(HalfLength * FVector(-1, +1, 0));
+	const FVector V2 = Location + Rotation.RotateVector(HalfLength * FVector(-1, -1, 0));
+	const FVector V3 = Location + Rotation.RotateVector(HalfLength * FVector(+1, -1, 0));
+
+	DrawDebugLinesPairsArrayView(Drawer, {
+		{ V0, V2 },
+		{ V1, V3 },
+		}, LineStyle, bDepthTest);
+}
+
 void UDrawDebugLibrary::DrawDebugDiamond(const FDebugDrawer& Drawer, const FVector& Location, const FRotator& Rotation, const FDrawDebugLineStyle& LineStyle, const bool bDepthTest, const float Radius)
 {
 	DrawDebugCircle(Drawer, Location, Rotation, LineStyle, bDepthTest, Radius, 4);
@@ -1280,6 +1388,46 @@ void UDrawDebugLibrary::DrawDebugRegularPolygon(const FDebugDrawer& Drawer, cons
 	DrawDebugCircle(Drawer, Location, Rotation, LineStyle, bDepthTest, Radius, Sides);
 }
 
+void UDrawDebugLibrary::DrawDebugCheckMark(const FDebugDrawer& Drawer, const FVector& Location, const FRotator& Rotation, const FDrawDebugLineStyle& LineStyle, const bool bDepthTest, const float Radius)
+{
+	const FVector P0 = Location + Rotation.RotateVector(Radius * FVector(0.0f, -0.2f, 0.0f));
+	const FVector P1 = Location + Rotation.RotateVector(Radius * FVector(-0.2f, 0.0f,0.0f));
+	const FVector P2 = Location + Rotation.RotateVector(Radius * FVector(0.8f, 0.6f, 0.0f));
+
+	DrawDebugLinesPairsArrayView(Drawer,
+		{
+			{ P0, P1 },
+			{ P1, P2 },
+		}, LineStyle, bDepthTest);
+}
+
+void UDrawDebugLibrary::DrawDebugCheckBox(const FDebugDrawer& Drawer, const FVector& Location, const FRotator& Rotation, const FDrawDebugLineStyle& LineStyle, const bool bDepthTest, const float HalfLength, const bool bChecked, const float TickRatio)
+{
+	DrawDebugSquare(Drawer, Location, Rotation, LineStyle, bDepthTest, HalfLength);
+	if (bChecked)
+	{
+		DrawDebugCheckMark(Drawer, Location, Rotation, LineStyle, bDepthTest, HalfLength * TickRatio);
+	}
+}
+
+void UDrawDebugLibrary::DrawDebugCrossBox(const FDebugDrawer& Drawer, const FVector& Location, const FRotator& Rotation, const FDrawDebugLineStyle& LineStyle, const bool bDepthTest, const float HalfLength, const bool bCrossed, const float CrossRatio)
+{
+	DrawDebugSquare(Drawer, Location, Rotation, LineStyle, bDepthTest, HalfLength);
+	if (bCrossed)
+	{
+		DrawDebugCross(Drawer, Location, Rotation, LineStyle, bDepthTest, HalfLength * CrossRatio);
+	}
+}
+
+void UDrawDebugLibrary::DrawDebugRadioButton(const FDebugDrawer& Drawer, const FVector& Location, const FRotator& Rotation, const FDrawDebugLineStyle& LineStyle, const bool bDepthTest, const float Radius, const bool bSelected, const float InnerButtonRatio, const int32 Segments)
+{
+	DrawDebugCircle(Drawer, Location, Rotation, LineStyle, bDepthTest, Radius, Segments);
+	if (bSelected)
+	{
+		DrawDebugCircle(Drawer, Location, Rotation, LineStyle, bDepthTest, InnerButtonRatio * Radius, Segments);
+	}
+}
+
 void UDrawDebugLibrary::DrawDebugCircleOutline(const FDebugDrawer& Drawer, const FVector& Location, const FRotator& Rotation, const FDrawDebugLineStyle& LineStyle, const bool bDepthTest, const float InnerRadius, const float OuterRadius, const int32 Segments)
 {
 	DrawDebugCircle(Drawer, Location, Rotation, LineStyle, bDepthTest, InnerRadius, Segments);
@@ -2665,7 +2813,7 @@ FVector UDrawDebugLibrary::DrawDebugStringViewDimensions(const FStringView Strin
 	XOffsetMax = FMath::Max(XOffsetMax, XOffset);
 	YOffset -= Settings.HeightScale * Settings.LineSpacing * Settings.Height;
 
-	return FVector(XOffsetMax, -YOffset, 0);
+	return FVector(0, XOffsetMax, -YOffset);
 }
 
 FVector UDrawDebugLibrary::DrawDebugTextDimensions(const FText& Text, const FDrawDebugStringSettings& Settings)
@@ -2678,6 +2826,26 @@ FVector UDrawDebugLibrary::DrawDebugNameDimensions(const FName& Name, const FDra
 	return DrawDebugStringViewDimensions(Name.ToString(), Settings);
 }
 
+FVector UDrawDebugLibrary::DrawDebugTextCenteringOffset(const FText& Text, const FDrawDebugStringSettings& Settings)
+{
+	return DrawDebugStringViewCenteringOffset(Text.ToString(), Settings);
+}
+
+FVector UDrawDebugLibrary::DrawDebugStringCenteringOffset(const FString& String, const FDrawDebugStringSettings& Settings)
+{
+	return DrawDebugStringViewCenteringOffset(String, Settings);
+}
+
+FVector UDrawDebugLibrary::DrawDebugStringViewCenteringOffset(const FStringView String, const FDrawDebugStringSettings& Settings)
+{
+	const FVector Dimensions = DrawDebugStringViewDimensions(String, Settings);
+	return FVector(0.0f, -Dimensions.Y / 2, Dimensions.Z / 2);
+}
+
+FVector UDrawDebugLibrary::DrawDebugNameCenteringOffset(const FName& Name, const FDrawDebugStringSettings& Settings)
+{
+	return DrawDebugStringViewCenteringOffset(Name.ToString(), Settings);
+}
 
 void UDrawDebugLibrary::DrawDebugString(const FDebugDrawer& Drawer, const FString& String, const FVector& Location, const FRotator& Rotation, const FDrawDebugLineStyle& LineStyle, const bool bDepthTest, const FDrawDebugStringSettings& Settings)
 {
@@ -2691,19 +2859,29 @@ void UDrawDebugLibrary::DrawDebugStringView(const FDebugDrawer& Drawer, const FS
 	TArray<TPair<FVector, FVector>> LineSegments;
 	LineSegments.Reserve(DrawDebugStringViewSegmentNum(String, Settings));
 
-	float XOffset = 0.0;
-	float YOffset = 0.0;
+	float InitialXOffset = 0.0f;
+	float InitialYOffset = 0.0f;
+
+	if (Settings.bCenterHorizontally || Settings.bCenterVertically)
+	{
+		const FVector Dimensions = DrawDebugStringViewDimensions(String, Settings);
+		if (Settings.bCenterHorizontally) { InitialXOffset -= Dimensions.Y / 2; }
+		if (Settings.bCenterVertically) { InitialYOffset += Dimensions.Z / 2; }
+	}
+
+	float XOffset = InitialXOffset;
+	float YOffset = InitialYOffset;
 
 	for (int32 CharIdx = 0; CharIdx < CharNum; CharIdx++)
 	{
 		if (String.GetData()[CharIdx] == '\n')
 		{
-			XOffset = 0.0;
+			XOffset = InitialXOffset;
 			YOffset -= Settings.LineSpacing * Settings.HeightScale * Settings.Height;
 		}
 		else if (String.GetData()[CharIdx] == '\r')
 		{
-			XOffset = 0.0;
+			XOffset = InitialXOffset;
 		}
 		else if (String.GetData()[CharIdx] == '\t')
 		{
@@ -2728,13 +2906,13 @@ void UDrawDebugLibrary::DrawDebugStringView(const FDebugDrawer& Drawer, const FS
 				const int32 YEndInt = (0x000000FF & (Segments[SegmentOff + SegmentIdx] >> 24));
 
 				const FVector Start = FVector(
-					XOffset + Settings.WidthScale * Settings.Height * (((float)XStartInt) / 128),
 					0.0f,
+					XOffset + Settings.WidthScale * Settings.Height * (((float)XStartInt) / 128),
 					YOffset - Settings.HeightScale * Settings.Height * (1.0f - ((float)YStartInt) / 128));
 
 				const FVector End = FVector(
-					XOffset + Settings.WidthScale * Settings.Height * (((float)XEndInt) / 128),
 					0.0f,
+					XOffset + Settings.WidthScale * Settings.Height * (((float)XEndInt) / 128),
 					YOffset - Settings.HeightScale * Settings.Height * (1.0f - ((float)YEndInt) / 128));
 
 				LineSegments.Add({ Rotation.RotateVector(Start) + Location, Rotation.RotateVector(End) + Location });
@@ -2894,43 +3072,22 @@ void UDrawDebugLibrary::DrawDebugGraphArrayView(
 		bDepthTest);
 }
 
-void UDrawDebugLibrary::DrawDebugGraphAxes(
+void UDrawDebugLibrary::DrawDebugGraphAxesLabels(
 	const FDebugDrawer& Drawer,
 	const FVector& Location,
 	const FRotator& Rotation,
 	const float XaxisLength,
 	const float YaxisLength,
 	const FDrawDebugLineStyle& TextLineStyle,
-	const FDrawDebugLineStyle& AxesLineStyle,
 	const bool bDepthTest,
 	const FDrawDebugGraphAxesSettings& AxesSettings)
 {
-	DrawDebugOrientedArrow(
-		Drawer,
-		Location,
-		(Rotation.Quaternion() * FQuat::MakeFromRotationVector(FVector(UE_PI / 2, 0, 0))).Rotator(),
-		XaxisLength,
-		AxesLineStyle,
-		bDepthTest,
-		AxesSettings.AxisArrowSettings);
-
-	DrawDebugOrientedArrow(
-		Drawer,
-		Location,
-		(Rotation.Quaternion() *
-			FQuat::MakeFromRotationVector(FVector(0, -UE_PI / 2, 0)) *
-			FQuat::MakeFromRotationVector(FVector(UE_PI / 2, 0, 0))).Rotator(),
-		YaxisLength,
-		AxesLineStyle,
-		bDepthTest,
-		AxesSettings.AxisArrowSettings);
-
 	const FVector TitleDimension = DrawDebugStringViewDimensions(AxesSettings.Title, AxesSettings.TitleSettings);
 
 	DrawDebugStringView(
 		Drawer,
 		AxesSettings.Title,
-		Location + Rotation.RotateVector(FVector(XaxisLength / 2.0f - TitleDimension.X / 2.0f, 0.0, YaxisLength + 1.1f * TitleDimension.Y)),
+		Location + Rotation.RotateVector(FVector(0.0, XaxisLength / 2.0f - TitleDimension.Y / 2.0f, YaxisLength + 1.1f * TitleDimension.Z)),
 		Rotation,
 		TextLineStyle,
 		bDepthTest,
@@ -2941,7 +3098,7 @@ void UDrawDebugLibrary::DrawDebugGraphAxes(
 	DrawDebugStringView(
 		Drawer,
 		AxesSettings.XaxisLabel,
-		Location + Rotation.RotateVector(FVector(XaxisLength / 2.0f - XlabelDimension.X / 2.0f, 0.0, -0.1f * XlabelDimension.Y)),
+		Location + Rotation.RotateVector(FVector(0.0, XaxisLength / 2.0f - XlabelDimension.Y / 2.0f, -0.1f * XlabelDimension.Z)),
 		Rotation,
 		TextLineStyle,
 		bDepthTest,
@@ -2952,13 +3109,76 @@ void UDrawDebugLibrary::DrawDebugGraphAxes(
 	DrawDebugStringView(
 		Drawer,
 		AxesSettings.YaxisLabel,
-		Location + Rotation.RotateVector(FVector(-YlabelDimension.Y - 0.1f * YlabelDimension.Y, 0.0, YaxisLength / 2 - YlabelDimension.X / 2)),
-		(Rotation.Quaternion() * FQuat::MakeFromRotationVector(FVector(0, -UE_PI / 2, 0))).Rotator(),
+		Location + Rotation.RotateVector(FVector(0.0, -YlabelDimension.Z - 0.1f * YlabelDimension.Z, YaxisLength / 2 - YlabelDimension.Y / 2)),
+		(Rotation.Quaternion() * FQuat::MakeFromRotationVector(FVector(UE_PI / 2, 0, 0))).Rotator(),
 		TextLineStyle,
 		bDepthTest,
 		AxesSettings.AxisLabelSettings);
 }
 
+void UDrawDebugLibrary::DrawDebugGraphAxes(
+	const FDebugDrawer& Drawer,
+	const FVector& Location,
+	const FRotator& Rotation,
+	const float XaxisLength,
+	const float YaxisLength,
+	const FDrawDebugLineStyle& TextLineStyle,
+	const FDrawDebugLineStyle& AxesLineStyle,
+	const bool bDepthTest,
+	const FDrawDebugGraphAxesSettings& AxesSettings)
+{
+	if (AxesSettings.bDrawAxesBox)
+	{
+		const FVector A0 = Location;
+		const FVector A1 = Location + Rotation.RotateVector(FVector(0, XaxisLength, 0));
+		const FVector A2 = Location + Rotation.RotateVector(FVector(0, XaxisLength, YaxisLength));
+		const FVector A3 = Location + Rotation.RotateVector(FVector(0, 0, YaxisLength));
+
+		DrawDebugLinesPairsArrayView(
+			Drawer,
+			{
+				{ A0, A1 },
+				{ A1, A2},
+				{ A2, A3 },
+				{ A2, A0 },
+			},
+			AxesLineStyle,
+			bDepthTest);
+	}
+	else
+	{
+		DrawDebugOrientedArrow(
+			Drawer,
+			Location,
+			(Rotation.Quaternion() *
+				FQuat::MakeFromRotationVector(FVector(0, 0, UE_PI / 2)) *
+				FQuat::MakeFromRotationVector(FVector(-UE_PI / 2, 0, 0))).Rotator(),
+			XaxisLength,
+			AxesLineStyle,
+			bDepthTest,
+			AxesSettings.AxisArrowSettings);
+
+		DrawDebugOrientedArrow(
+			Drawer,
+			Location,
+			(Rotation.Quaternion() * FQuat::MakeFromRotationVector(FVector(0, -UE_PI / 2, 0))).Rotator(),
+			YaxisLength,
+			AxesLineStyle,
+			bDepthTest,
+			AxesSettings.AxisArrowSettings);
+	}
+
+	DrawDebugGraphAxesLabels(
+		Drawer,
+		Location,
+		Rotation,
+		XaxisLength,
+		YaxisLength,
+		TextLineStyle,
+		bDepthTest,
+		AxesSettings);
+}
+
 void UDrawDebugLibrary::DrawDebugGraphLine(
 	const FDebugDrawer& Drawer,
 	const FVector& Location,
@@ -3008,8 +3228,8 @@ void UDrawDebugLibrary::DrawDebugGraphLineArrayView(
 		const float Yloc1 = YaxisLength * ((Yvalues[PointIdx + 1] - Ymin) / YScale);
 
 		LineSegments.Add({
-			Location + Rotation.RotateVector(FVector(Xloc0, 0.0, Yloc0)),
-			Location + Rotation.RotateVector(FVector(Xloc1, 0.0, Yloc1)) });
+			Location + Rotation.RotateVector(FVector(0.0, Xloc0, Yloc0)),
+			Location + Rotation.RotateVector(FVector(0.0, Xloc1, Yloc1)) });
 	}
 
 	DrawDebugLinesPairsArrayView(Drawer, LineSegments, LineStyle, bDepthTest);
@@ -3051,26 +3271,26 @@ void UDrawDebugLibrary::DrawDebugGraphLegendArrayView(
 	float TotalHeight = 0.0f;
 	for (int32 LegendIdx = 0; LegendIdx < LegendNum; LegendIdx++)
 	{
-		TotalHeight = DrawDebugStringViewDimensions(LegendLabels[LegendIdx], LegendSettings).Y;
+		TotalHeight = DrawDebugStringViewDimensions(LegendLabels[LegendIdx], LegendSettings).Z;
 	}
 
 	float VerticalOffset = TotalHeight / 2.0f + YaxisLength / 2.0f;
 	for (int32 LegendIdx = 0; LegendIdx < LegendNum; LegendIdx++)
 	{
-		const float LegendHeight = DrawDebugStringViewDimensions(LegendLabels[LegendIdx], LegendSettings).Y;
+		const float LegendHeight = DrawDebugStringViewDimensions(LegendLabels[LegendIdx], LegendSettings).Z;
 
-		const FQuat IconRotation = FQuat::MakeFromRotationVector(FVector(UE_PI / 2.0, 0.0f, 0.0f));
+		const FQuat IconRotation = FQuat::MakeFromRotationVector(FVector(0.0f, UE_PI / 2.0, 0.0f));
 
 		DrawDebugDiamond(
 			Drawer, 
-			Location + Rotation.RotateVector(FVector(XaxisLength + 2.0f * IconSize, 0.0, VerticalOffset - LegendHeight / 2.0f)),
+			Location + Rotation.RotateVector(FVector(0.0, XaxisLength + 2.0f * IconSize, VerticalOffset - LegendHeight / 2.0f)),
 			(Rotation.Quaternion() * IconRotation).Rotator(),
 			DrawDebugLineStyleWithColor(IconLineStyle, LegendColors[LegendIdx]), bDepthTest, IconSize / 2.0f);
 
 		DrawDebugStringView(
 			Drawer,
 			LegendLabels[LegendIdx],
-			Location + Rotation.RotateVector(FVector(XaxisLength + 4.0f * IconSize, 0.0, VerticalOffset)),
+			Location + Rotation.RotateVector(FVector(0.0, XaxisLength + 4.0f * IconSize, VerticalOffset)),
 			Rotation,
 			DrawDebugLineStyleWithColor(TextLineStyle, LegendColors[LegendIdx]),
 			bDepthTest,
diff --git a/Engine/Plugins/Experimental/DrawDebugLibrary/Source/DrawDebugLibrary/Public/DrawDebugLibrary.h b/Engine/Plugins/Experimental/DrawDebugLibrary/Source/DrawDebugLibrary/Public/DrawDebugLibrary.h
--- a/Engine/Plugins/Experimental/DrawDebugLibrary/Source/DrawDebugLibrary/Public/DrawDebugLibrary.h
+++ b/Engine/Plugins/Experimental/DrawDebugLibrary/Source/DrawDebugLibrary/Public/DrawDebugLibrary.h
@@ -422,6 +422,14 @@ struct FDrawDebugStringSettings
 	/** Gaps between characters in cm */
 	UPROPERTY(BlueprintReadWrite, EditAnywhere, Category = "Draw Debug Library", meta = (ClampMin = "0", UIMin = "0", Units = "cm"))
 	float CharacterSpacing = 0.0f;
+
+	/** Draws the string centered vertically rather than starting from the top left */
+	UPROPERTY(BlueprintReadWrite, EditAnywhere, Category = "Draw Debug Library")
+	bool bCenterVertically = false;
+
+	/** Draws the string centered horizontally rather than starting from the top left */
+	UPROPERTY(BlueprintReadWrite, EditAnywhere, Category = "Draw Debug Library")
+	bool bCenterHorizontally = false;
 };
 
 /** Settings for debug drawing graphs axes */
@@ -453,6 +461,10 @@ struct FDrawDebugGraphAxesSettings
 	/** Settings for the axis labels */
 	UPROPERTY(BlueprintReadWrite, EditAnywhere, Category = "Draw Debug Library")
 	FDrawDebugStringSettings AxisLabelSettings;
+
+	/** Draws the axes as a box instead of two arrows */
+	UPROPERTY(BlueprintReadWrite, EditAnywhere, Category = "Draw Debug Library")
+	bool bDrawAxesBox = false;
 };
 
 /**
@@ -667,6 +679,55 @@ public: // Convenience functions
 	UFUNCTION(BlueprintPure = false, Category = "Draw Debug Library", meta = (NotBlueprintThreadSafe))
 	static UE_API void AddToVectorHistoryArray(UPARAM(ref) TArray<FVector>& InOutValues, const FVector NewValue = FVector::ZeroVector, const int32 MaxHistoryNum = 60);
 
+	/** Convenience function that adds a name to the end of an array, popping values from the front once the max is reached */
+	UFUNCTION(BlueprintPure = false, Category = "Draw Debug Library", meta = (NotBlueprintThreadSafe))
+	static UE_API void AddToNameHistoryArray(UPARAM(ref) TArray<FName>& InOutValues, const FName NewValue = NAME_None, const int32 MaxHistoryNum = 60);
+
+	/** Returns a new location and rotation after applying a local offset to the input draw location and rotation */
+	UFUNCTION(BlueprintPure, Category = "Draw Debug Library")
+	static UE_API void DrawDebugLocalOffset(
+		FVector& OutLocation,
+		FRotator& OutRotation,
+		const FVector DrawLocation,
+		const FRotator DrawRotation,
+		const FVector DrawOffset);
+
+	/** Create a new location and rotation that orients debug drawing of upright objects such as text onto the floor */
+	UFUNCTION(BlueprintPure, Category = "Draw Debug Library")
+	static UE_API void DrawDebugOrientUprightToFloor(
+		FVector& OutLocation,
+		FRotator& OutRotation,
+		const FVector DrawWorldLocation,
+		const FVector DrawUprightOffset);
+
+	/** Create a new location and rotation that orients debug drawing of floor objects such as shapes upright */
+	UFUNCTION(BlueprintPure, Category = "Draw Debug Library")
+	static UE_API void DrawDebugOrientFloorToUpright(
+		FVector& OutLocation,
+		FRotator& OutRotation,
+		const FVector DrawWorldLocation,
+		const FVector DrawFloorOffset);
+
+	/** Create a new location and rotation that orients debug drawing of upright objects such as text to face the camera */
+	UFUNCTION(BlueprintPure, Category = "Draw Debug Library")
+	static UE_API void DrawDebugOrientUprightToCamera(
+		FVector& OutLocation,
+		FRotator& OutRotation,
+		const FVector DrawWorldLocation,
+		const FVector DrawUprightOffset,
+		const FRotator CameraRotation,
+		const bool bUseOnlyYaw = false);
+
+	/** Create a new location and rotation that orients debug drawing of floor objects such as shapes to face the camera */
+	UFUNCTION(BlueprintPure, Category = "Draw Debug Library")
+	static UE_API void DrawDebugOrientFloorToCamera(
+		FVector& OutLocation,
+		FRotator& OutRotation,
+		const FVector DrawWorldLocation,
+		const FVector DrawFloorOffset,
+		const FRotator CameraRotation,
+		const bool bUseOnlyYaw = false);
+
 public: // Style Modifiers
 
 	/** Make a point style from a color */
@@ -783,6 +844,10 @@ public: // Basic Shapes
 	UFUNCTION(BlueprintCallable, Category = "Draw Debug Library", meta = (DevelopmentOnly, AutoCreateRefTerm = "Location, Rotation, LineStyle"))
 	static UE_API void DrawDebugSquare(const FDebugDrawer& Drawer, const FVector& Location, const FRotator& Rotation, const FDrawDebugLineStyle& LineStyle = FDrawDebugLineStyle(), const bool bDepthTest = true, const float HalfLength = 10.0f);
 
+	/** Debug Draw a cross on the XY axis */
+	UFUNCTION(BlueprintCallable, Category = "Draw Debug Library", meta = (DevelopmentOnly, AutoCreateRefTerm = "Location, Rotation, LineStyle"))
+	static UE_API void DrawDebugCross(const FDebugDrawer& Drawer, const FVector& Location, const FRotator& Rotation, const FDrawDebugLineStyle& LineStyle = FDrawDebugLineStyle(), const bool bDepthTest = true, const float HalfLength = 10.0f);
+
 	/** Debug Draw a diamond on the XY axis */
 	UFUNCTION(BlueprintCallable, Category = "Draw Debug Library", meta = (DevelopmentOnly, AutoCreateRefTerm = "Location, Rotation, LineStyle"))
 	static UE_API void DrawDebugDiamond(const FDebugDrawer& Drawer, const FVector& Location, const FRotator& Rotation, const FDrawDebugLineStyle& LineStyle = FDrawDebugLineStyle(), const bool bDepthTest = true, const float Radius = 10.0f);
@@ -815,6 +880,22 @@ public: // Basic Shapes
 	UFUNCTION(BlueprintCallable, Category = "Draw Debug Library", meta = (DevelopmentOnly, AutoCreateRefTerm = "Location, Rotation, LineStyle"))
 	static UE_API void DrawDebugRegularPolygon(const FDebugDrawer& Drawer, const FVector& Location, const FRotator& Rotation, const FDrawDebugLineStyle& LineStyle = FDrawDebugLineStyle(), const bool bDepthTest = true, const float Radius = 10.0f, const int32 Sides = 11);
 
+	/** Debug Draw a check mark on the XY axis */
+	UFUNCTION(BlueprintCallable, Category = "Draw Debug Library", meta = (DevelopmentOnly, AutoCreateRefTerm = "Location, Rotation, LineStyle"))
+	static UE_API void DrawDebugCheckMark(const FDebugDrawer& Drawer, const FVector& Location, const FRotator& Rotation, const FDrawDebugLineStyle& LineStyle = FDrawDebugLineStyle(), const bool bDepthTest = true, const float Radius = 10.0f);
+
+	/** Debug Draw a checkbox on the XY axis */
+	UFUNCTION(BlueprintCallable, Category = "Draw Debug Library", meta = (DevelopmentOnly, AutoCreateRefTerm = "Location, Rotation, LineStyle"))
+	static UE_API void DrawDebugCheckBox(const FDebugDrawer& Drawer, const FVector& Location, const FRotator& Rotation, const FDrawDebugLineStyle& LineStyle = FDrawDebugLineStyle(), const bool bDepthTest = true, const float HalfLength = 10.0f, const bool bChecked = false, const float TickRatio = 1.75f);
+
+	/** Debug Draw a cross box on the XY axis */
+	UFUNCTION(BlueprintCallable, Category = "Draw Debug Library", meta = (DevelopmentOnly, AutoCreateRefTerm = "Location, Rotation, LineStyle"))
+	static UE_API void DrawDebugCrossBox(const FDebugDrawer& Drawer, const FVector& Location, const FRotator& Rotation, const FDrawDebugLineStyle& LineStyle = FDrawDebugLineStyle(), const bool bDepthTest = true, const float HalfLength = 10.0f, const bool bCrossed = false, const float CrossRatio = 1.0);
+
+	/** Debug Draw a radio button on the XY axis */
+	UFUNCTION(BlueprintCallable, Category = "Draw Debug Library", meta = (DevelopmentOnly, AutoCreateRefTerm = "Location, Rotation, LineStyle"))
+	static UE_API void DrawDebugRadioButton(const FDebugDrawer& Drawer, const FVector& Location, const FRotator& Rotation, const FDrawDebugLineStyle& LineStyle = FDrawDebugLineStyle(), const bool bDepthTest = true, const float Radius = 10.0f, const bool bSelected = false, const float InnerButtonRatio = 0.5f, const int32 Segments = 17);
+
 	/** Debug Draw an axis-aligned cross */
 	UFUNCTION(BlueprintCallable, Category = "Draw Debug Library", meta = (DevelopmentOnly, AutoCreateRefTerm = "Location, Rotation, LineStyle"))
 	static UE_API void DrawDebugLocator(const FDebugDrawer& Drawer, const FVector& Location, const FRotator& Rotation, const FDrawDebugLineStyle& LineStyle = FDrawDebugLineStyle(), const bool bDepthTest = true, const float Radius = 10.0f);
@@ -1052,23 +1133,36 @@ public: // Poses / Skeletons
 public: // Text
 
 	/** Get the number of line segments required to draw debug string. */
-	UFUNCTION(BlueprintCallable, Category = "Draw Debug Library", meta = (DevelopmentOnly, AutoCreateRefTerm = "String, Settings"))
+	UFUNCTION(BlueprintPure, Category = "Draw Debug Library", meta = (DevelopmentOnly, AutoCreateRefTerm = "String, Settings"))
 	static UE_API int32 DrawDebugStringSegmentNum(const FString& String, const FDrawDebugStringSettings& Settings = FDrawDebugStringSettings());
 	static UE_API int32 DrawDebugStringViewSegmentNum(const FStringView String, const FDrawDebugStringSettings& Settings = FDrawDebugStringSettings());
 
 	/** Get the dimensions of a draw debug string. Useful for aligning or centering text. */
-	UFUNCTION(BlueprintCallable, Category = "Draw Debug Library", meta = (DevelopmentOnly, AutoCreateRefTerm = "String, Settings"))
+	UFUNCTION(BlueprintPure, Category = "Draw Debug Library", meta = (DevelopmentOnly, AutoCreateRefTerm = "String, Settings"))
 	static UE_API FVector DrawDebugStringDimensions(const FString& String, const FDrawDebugStringSettings& Settings = FDrawDebugStringSettings());
 	static UE_API FVector DrawDebugStringViewDimensions(const FStringView String, const FDrawDebugStringSettings& Settings = FDrawDebugStringSettings());
 
 	/** Get the dimensions of a draw debug text. Useful for aligning or centering text. */
-	UFUNCTION(BlueprintCallable, Category = "Draw Debug Library", meta = (DevelopmentOnly, AutoCreateRefTerm = "Text, Settings"))
+	UFUNCTION(BlueprintPure, Category = "Draw Debug Library", meta = (DevelopmentOnly, AutoCreateRefTerm = "Text, Settings"))
 	static UE_API FVector DrawDebugTextDimensions(const FText& Text, const FDrawDebugStringSettings& Settings = FDrawDebugStringSettings());
 
 	/** Get the dimensions of a draw debug name. Useful for aligning or centering text. */
-	UFUNCTION(BlueprintCallable, Category = "Draw Debug Library", meta = (DevelopmentOnly, AutoCreateRefTerm = "Name, Settings"))
+	UFUNCTION(BlueprintPure, Category = "Draw Debug Library", meta = (DevelopmentOnly, AutoCreateRefTerm = "Name, Settings"))
 	static UE_API FVector DrawDebugNameDimensions(const FName& Name, const FDrawDebugStringSettings& Settings = FDrawDebugStringSettings());
 
+	/** Get the local offset required to center a draw debug text. */
+	UFUNCTION(BlueprintPure, Category = "Draw Debug Library", meta = (DevelopmentOnly, AutoCreateRefTerm = "Text, Settings"))
+	static UE_API FVector DrawDebugTextCenteringOffset(const FText& Text, const FDrawDebugStringSettings& Settings = FDrawDebugStringSettings());
+
+	/** Get the local offset required to center a draw debug string. */
+	UFUNCTION(BlueprintPure, Category = "Draw Debug Library", meta = (DevelopmentOnly, AutoCreateRefTerm = "String, Settings"))
+	static UE_API FVector DrawDebugStringCenteringOffset(const FString& String, const FDrawDebugStringSettings& Settings = FDrawDebugStringSettings());
+	static UE_API FVector DrawDebugStringViewCenteringOffset(const FStringView String, const FDrawDebugStringSettings& Settings = FDrawDebugStringSettings());
+
+	/** Get the local offset required to center a draw debug name. */
+	UFUNCTION(BlueprintPure, Category = "Draw Debug Library", meta = (DevelopmentOnly, AutoCreateRefTerm = "Name, Settings"))
+	static UE_API FVector DrawDebugNameCenteringOffset(const FName& Name, const FDrawDebugStringSettings& Settings = FDrawDebugStringSettings());
+
 	/** Debug Draw a string. Will only render ASCII characters. */
 	UFUNCTION(BlueprintCallable, Category = "Draw Debug Library", meta = (DevelopmentOnly, AutoCreateRefTerm = "String, Location, Rotation, LineStyle, Settings"))
 	static UE_API void DrawDebugString(const FDebugDrawer& Drawer, const FString& String, const FVector& Location, const FRotator& Rotation, const FDrawDebugLineStyle& LineStyle = FDrawDebugLineStyle(), const bool bDepthTest = true, const FDrawDebugStringSettings& Settings = FDrawDebugStringSettings());
@@ -1137,6 +1231,18 @@ public: // Graphs
 		const bool bDepthTest = true,
 		const FDrawDebugGraphAxesSettings& AxesSettings = FDrawDebugGraphAxesSettings());
 
+	/** Debug Draw a simple graph axes labels. */
+	UFUNCTION(BlueprintCallable, Category = "Draw Debug Library", meta = (DevelopmentOnly, AutoCreateRefTerm = "Location, Rotation, TextLineStyle, AxesSettings"))
+	static UE_API void DrawDebugGraphAxesLabels(
+		const FDebugDrawer& Drawer,
+		const FVector& Location,
+		const FRotator& Rotation,
+		const float XaxisLength = 100.0f,
+		const float YaxisLength = 100.0f,
+		const FDrawDebugLineStyle& TextLineStyle = FDrawDebugLineStyle(),
+		const bool bDepthTest = true,
+		const FDrawDebugGraphAxesSettings& AxesSettings = FDrawDebugGraphAxesSettings());
+
 	/** Debug Draw a simple graph axes. */
 	UFUNCTION(BlueprintCallable, Category = "Draw Debug Library", meta = (DevelopmentOnly, AutoCreateRefTerm = "Location, Rotation, TextLineStyle, AxesLineStyle, AxesSettings"))
 	static UE_API void DrawDebugGraphAxes(
