DrawDebugLibrary: Started adding some RigUnits for debug drawing.

#Synced-CL 48409205

diff --git a/Engine/Plugins/Experimental/DrawDebugLibrary/Source/DrawDebugLibrary/Private/DrawDebugLibrary.cpp b/Engine/Plugins/Experimental/DrawDebugLibrary/Source/DrawDebugLibrary/Private/DrawDebugLibrary.cpp
--- a/Engine/Plugins/Experimental/DrawDebugLibrary/Source/DrawDebugLibrary/Private/DrawDebugLibrary.cpp
+++ b/Engine/Plugins/Experimental/DrawDebugLibrary/Source/DrawDebugLibrary/Private/DrawDebugLibrary.cpp
@@ -11,18 +11,33 @@
 #include "Animation/AnimInstanceProxy.h"
 #include "VisualLogger/VisualLogger.h"
 #include "Components/SkinnedMeshComponent.h"
+#include "RigVMCore/RigVMRegistry.h"
 #include "RigVM/Public/RigVMCore/RigVMDrawInterface.h"
 
+#include UE_INLINE_GENERATED_CPP_BY_NAME(DrawDebugLibrary)
+
 void FDrawDebugLibraryModule::StartupModule()
 {
-	// This code will execute after your module is loaded into memory; the exact timing is specified in the .uplugin file per-module
-}
+	static UScriptStruct* const AllowedStructTypes[] =
+	{
+		FDrawDebugPointStyle::StaticStruct(),
+		FDrawDebugLineStyle::StaticStruct(),
+		FDrawDebugArrowSettings::StaticStruct(),
+		FDrawDebugChairSettings::StaticStruct(),
+		FDrawDebugBackpackSettings::StaticStruct(),
+		FDrawDebugBriefcaseSettings::StaticStruct(),
+		FDrawDebugDoorSettings::StaticStruct(),
+		FDrawDebugSkeletonSettings::StaticStruct(),
+		FDrawDebugStringSettings::StaticStruct(),
+		FDrawDebugGraphAxesSettings::StaticStruct(),
+		FDebugDrawer::StaticStruct(),
+	};
 
-void FDrawDebugLibraryModule::ShutdownModule()
-{
-	// This function may be called during shutdown to clean up your module.  For modules that support dynamic reloading,
-	// we call this function before unloading the module.
+	FRigVMRegistry& RigVMRegistry = FRigVMRegistry::Get();
+	RigVMRegistry.RegisterStructTypes(AllowedStructTypes);
 }
+
+void FDrawDebugLibraryModule::ShutdownModule() {}
 	
 IMPLEMENT_MODULE(FDrawDebugLibraryModule, DrawDebugLibrary)
 
@@ -476,7 +491,7 @@ FDebugDrawer FDebugDrawer::MakeDebugDrawer(FRigVMDrawInterface* InRigVMDrawInter
 }
 
 FDebugDrawer FDebugDrawer::MakeVisualLoggerDebugDrawer(
-	UObject* Object, 
+	const UObject* Object,
 	const FName Category, 
 	const EDrawDebugLogVerbosity Verbosity, 
 	const bool bDrawToScene,
@@ -492,7 +507,7 @@ FDebugDrawer FDebugDrawer::MakeVisualLoggerDebugDrawer(
 }
 
 FDebugDrawer FDebugDrawer::MakeVisualLoggerDebugDrawer(
-	UObject* Object, 
+	const UObject* Object,
 	const FLogCategoryBase& Category, 
 	const ELogVerbosity::Type Verbosity, 
 	const bool bDrawToScene,
@@ -796,17 +811,22 @@ FDebugDrawer UDrawDebugLibrary::MakePDIDebugDrawer(FPrimitiveDrawInterface* Prim
 	return FDebugDrawer::MakeDebugDrawer(PrimitiveDrawInterface);
 }
 
-FDebugDrawer UDrawDebugLibrary::MakeVisualLoggerDebugDrawer(UObject* Object, const FName Category, const EDrawDebugLogVerbosity Verbosity, const bool bDrawToScene,const bool bDrawToSceneWhileRecording)
+FDebugDrawer UDrawDebugLibrary::MakeRigVMDebugDrawer(FRigVMDrawInterface* InRigVMDrawInterface)
+{
+	return FDebugDrawer::MakeDebugDrawer(InRigVMDrawInterface);
+}
+
+FDebugDrawer UDrawDebugLibrary::MakeVisualLoggerDebugDrawer(const UObject* Object, const FName Category, const EDrawDebugLogVerbosity Verbosity, const bool bDrawToScene,const bool bDrawToSceneWhileRecording)
 {
 	return FDebugDrawer::MakeVisualLoggerDebugDrawer(Object, Category, Verbosity, bDrawToScene, bDrawToSceneWhileRecording);
 }
 
-FDebugDrawer UDrawDebugLibrary::MakeVisualLoggerDebugDrawerFromObject(UObject* Object, const FName Category, const EDrawDebugLogVerbosity Verbosity, const bool bDrawToScene, const bool bDrawToSceneWhileRecording)
+FDebugDrawer UDrawDebugLibrary::MakeVisualLoggerDebugDrawerFromObject(const UObject* Object, const FName Category, const EDrawDebugLogVerbosity Verbosity, const bool bDrawToScene, const bool bDrawToSceneWhileRecording)
 {
 	return FDebugDrawer::MakeVisualLoggerDebugDrawer(Object, Category, Verbosity, bDrawToScene, bDrawToSceneWhileRecording);
 }
 
-FDebugDrawer UDrawDebugLibrary::MakeVisualLoggerDebugDrawerFromObjectWithCategory(UObject* Object, const FLogCategoryBase& Category, const ELogVerbosity::Type Verbosity, const bool bDrawToScene, const bool bDrawToSceneWhileRecording)
+FDebugDrawer UDrawDebugLibrary::MakeVisualLoggerDebugDrawerFromObjectWithCategory(const UObject* Object, const FLogCategoryBase& Category, const ELogVerbosity::Type Verbosity, const bool bDrawToScene, const bool bDrawToSceneWhileRecording)
 {
 	return FDebugDrawer::MakeVisualLoggerDebugDrawer(Object, Category, Verbosity, bDrawToScene, bDrawToSceneWhileRecording);
 }
diff --git a/Engine/Plugins/Experimental/DrawDebugLibrary/Source/DrawDebugLibrary/Private/RigUnit_DrawDebugLibrary.cpp b/Engine/Plugins/Experimental/DrawDebugLibrary/Source/DrawDebugLibrary/Private/RigUnit_DrawDebugLibrary.cpp
new file mode 100644
--- /dev/null
+++ b/Engine/Plugins/Experimental/DrawDebugLibrary/Source/DrawDebugLibrary/Private/RigUnit_DrawDebugLibrary.cpp
@@ -0,0 +1,55 @@
+// Copyright Epic Games, Inc. All Rights Reserved.
+
+#include "RigUnit_DrawDebugLibrary.h"
+
+#include UE_INLINE_GENERATED_CPP_BY_NAME(RigUnit_DrawDebugLibrary)
+
+FRigVMFunction_MakeNullDebugDrawer_Execute()
+{
+	Drawer = UDrawDebugLibrary::MakeNullDebugDrawer();
+}
+
+FRigVMFunction_MakeDebugDrawer_Execute()
+{
+	Drawer = UDrawDebugLibrary::MakeRigVMDebugDrawer(ExecuteContext.GetDrawInterface());
+}
+
+FRigVMFunction_MakeVisualLoggerDebugDrawer_Execute()
+{
+	Drawer = UDrawDebugLibrary::MakeVisualLoggerDebugDrawer(ExecuteContext.GetOwningObject(), Category, Verbosity, bDrawToScene, bDrawToSceneWhileRecording);
+}
+
+FRigVMFunction_MakeMergedDebugDrawer_Execute()
+{
+	Drawer = UDrawDebugLibrary::MakeMergedDebugDrawerArrayView(Drawers);
+}
+
+FRigUnit_DrawDebugLine_Execute()
+{
+    UDrawDebugLibrary::DrawDebugLine(Drawer, StartLocation, EndLocation, LineStyle, bDepthTest);
+}
+
+FRigUnit_DrawDebugLocation_Execute()
+{
+	UDrawDebugLibrary::DrawDebugLocation(Drawer, Location, LineStyle, bDepthTest, DrawRadius, Segments);
+}
+
+FRigUnit_DrawDebugVelocity_Execute()
+{
+	UDrawDebugLibrary::DrawDebugVelocity(Drawer, Location, Velocity, LineStyle, bDepthTest, DrawVelocityLineScale);
+}
+
+FRigUnit_DrawDebugTransformTrajectory_Execute()
+{
+	UDrawDebugLibrary::DrawDebugTransformTrajectory(Drawer, TransformTrajectory, LineStyle, bDepthTest, Radius, VerticalOffset);
+}
+
+FRigUnit_DrawDebugString_Execute()
+{
+	UDrawDebugLibrary::DrawDebugString(Drawer, String, Location, Rotation, LineStyle, bDepthTest, Settings);
+}
+
+FRigUnit_DrawDebugName_Execute()
+{
+	UDrawDebugLibrary::DrawDebugName(Drawer, Name, Location, Rotation, LineStyle, bDepthTest, Settings);
+}
diff --git a/Engine/Plugins/Experimental/DrawDebugLibrary/Source/DrawDebugLibrary/Public/DrawDebugLibrary.h b/Engine/Plugins/Experimental/DrawDebugLibrary/Source/DrawDebugLibrary/Public/DrawDebugLibrary.h
--- a/Engine/Plugins/Experimental/DrawDebugLibrary/Source/DrawDebugLibrary/Public/DrawDebugLibrary.h
+++ b/Engine/Plugins/Experimental/DrawDebugLibrary/Source/DrawDebugLibrary/Public/DrawDebugLibrary.h
@@ -523,7 +523,7 @@ public: // Constructors
 
 	/** Make a debug drawer for the Visual Logger */
 	static UE_API FDebugDrawer MakeVisualLoggerDebugDrawer(
-		UObject* Object, 
+		const UObject* Object,
 		const FName Category = TEXT("LogDrawDebugLibrary"), 
 		const EDrawDebugLogVerbosity Verbosity = EDrawDebugLogVerbosity::Display, 
 		const bool bDrawToScene = true,
@@ -531,7 +531,7 @@ public: // Constructors
 
 	/** Make a debug drawer for the Visual Logger */
 	static UE_API FDebugDrawer MakeVisualLoggerDebugDrawer(
-		UObject* Object, 
+		const UObject* Object, 
 		const FLogCategoryBase& Category, 
 		const ELogVerbosity::Type Verbosity = ELogVerbosity::Display, 
 		const bool bDrawToScene = true,
@@ -569,7 +569,7 @@ private:
 
 	/** Visual Logger Object - used for writing to the visual logger. */
 	UPROPERTY()
-	TObjectPtr<UObject> VisualLoggerObject = nullptr;
+	TObjectPtr<const UObject> VisualLoggerObject = nullptr;
 
 	/** Visual Logger Category. */
 	FName VisualLoggerCategory = NAME_None;
@@ -624,10 +624,13 @@ public: // Drawers
 	/** Make a debug drawer for a FPrimitiveDrawInterface */
 	static UE_API FDebugDrawer MakePDIDebugDrawer(FPrimitiveDrawInterface* PrimitiveDrawInterface);
 
+	/** Make a debug drawer for a FRigVMDrawInterface */
+	static UE_API FDebugDrawer MakeRigVMDebugDrawer(FRigVMDrawInterface* InRigVMDrawInterface);
+
 	/** Make a debug drawer for the Visual Logger from the current "Self" object */
 	UFUNCTION(BlueprintPure, Category = "Draw Debug Library", meta = (DevelopmentOnly, DefaultToSelf = "Object", HidePin = "Object"))
 	static UE_API FDebugDrawer MakeVisualLoggerDebugDrawer(
-		UObject* Object, 
+		const UObject* Object,
 		const FName Category = TEXT("LogDrawDebugLibrary"), 
 		const EDrawDebugLogVerbosity Verbosity = EDrawDebugLogVerbosity::Display, 
 		const bool bDrawToScene = true,
@@ -636,14 +639,14 @@ public: // Drawers
 	/** Make a debug drawer for the Visual Logger */
 	UFUNCTION(BlueprintPure, Category = "Draw Debug Library", meta = (DevelopmentOnly))
 	static UE_API FDebugDrawer MakeVisualLoggerDebugDrawerFromObject(
-		UObject* Object, 
+		const UObject* Object,
 		const FName Category = TEXT("LogDrawDebugLibrary"), 
 		const EDrawDebugLogVerbosity Verbosity = EDrawDebugLogVerbosity::Display, 
 		const bool bDrawToScene = true, 
 		const bool bDrawToSceneWhileRecording = true);
 
 	static UE_API FDebugDrawer MakeVisualLoggerDebugDrawerFromObjectWithCategory(
-		UObject* Object, 
+		const UObject* Object,
 		const FLogCategoryBase& Category, 
 		const ELogVerbosity::Type Verbosity = ELogVerbosity::Display, 
 		const bool bDrawToScene = true, 
diff --git a/Engine/Plugins/Experimental/DrawDebugLibrary/Source/DrawDebugLibrary/Public/RigUnit_DrawDebugLibrary.h b/Engine/Plugins/Experimental/DrawDebugLibrary/Source/DrawDebugLibrary/Public/RigUnit_DrawDebugLibrary.h
new file mode 100644
--- /dev/null
+++ b/Engine/Plugins/Experimental/DrawDebugLibrary/Source/DrawDebugLibrary/Public/RigUnit_DrawDebugLibrary.h
@@ -0,0 +1,341 @@
+// Copyright Epic Games, Inc. All Rights Reserved.
+
+#pragma once
+
+#include "DrawDebugLibrary.h"
+#include "RigVMCore/RigVMStruct.h"
+#include "Animation/TrajectoryTypes.h"
+
+#include "RigUnit_DrawDebugLibrary.generated.h"
+
+#define UE_API DRAWDEBUGLIBRARY_API
+
+/*
+ * The base class for all draw debug library functions
+ */
+USTRUCT(meta = (Abstract, Category = "Draw Debug Library", NodeColor = "0.83077 0.846873 0.049707", DocumentationPolicy = "Strict"))
+struct FRigVMFunction_DrawDebugLibraryBase : public FRigVMStruct
+{
+	GENERATED_BODY()
+
+	virtual void Execute() {};
+};
+
+/*
+ * Base class for all draw debug library rig units
+ */
+USTRUCT(meta=(Abstract, ExecuteContext = "FRigVMExecuteContext", Category="Draw Debug Library", NodeColor = "0.83077 0.846873 0.049707", DocumentationPolicy = "Strict"))
+struct FRigUnit_DrawDebugLibraryBase : public FRigVMStruct
+{
+	GENERATED_BODY()
+};
+
+// Makes a null debug drawer
+USTRUCT(meta = (DisplayName = "Make Null Debug Drawer"))
+struct FRigVMFunction_MakeNullDebugDrawer : public FRigVMFunction_DrawDebugLibraryBase
+{
+	GENERATED_BODY()
+
+	RIGVM_METHOD()
+	UE_API virtual void Execute() override;
+
+	// Debug Drawer
+	UPROPERTY(meta = (Output))
+	FDebugDrawer Drawer;
+};
+
+// Makes a normal RigVM debug drawer
+USTRUCT(meta = (DisplayName = "Make Debug Drawer"))
+struct FRigVMFunction_MakeDebugDrawer : public FRigVMFunction_DrawDebugLibraryBase
+{
+	GENERATED_BODY()
+
+	RIGVM_METHOD()
+	UE_API virtual void Execute() override;
+
+	// Debug Drawer
+	UPROPERTY(meta = (Output))
+	FDebugDrawer Drawer;
+};
+
+// Makes a debug drawer for the visual logger.
+USTRUCT(meta = (DisplayName = "Make Visual Logger Debug Drawer"))
+struct FRigVMFunction_MakeVisualLoggerDebugDrawer : public FRigVMFunction_DrawDebugLibraryBase
+{
+	GENERATED_BODY()
+
+	RIGVM_METHOD()
+	UE_API virtual void Execute() override;
+
+	// Debug Drawer
+	UPROPERTY(meta = (Output))
+	FDebugDrawer Drawer;
+
+	// Log Category
+	UPROPERTY(meta = (Input))
+	FName Category = TEXT("LogDrawDebugLibrary");
+
+	// Log Verbosity
+	UPROPERTY(meta = (Input))
+	EDrawDebugLogVerbosity Verbosity = EDrawDebugLogVerbosity::Display;
+
+	// If to draw to the scene also
+	UPROPERTY(meta = (Input))
+	bool bDrawToScene = true;
+
+	// If to draw to the scene while recording
+	UPROPERTY(meta = (Input))
+	bool bDrawToSceneWhileRecording = true;
+};
+
+// Makes a merged debug drawer
+USTRUCT(meta = (DisplayName = "Make Merged Debug Drawer"))
+struct FRigVMFunction_MakeMergedDebugDrawer : public FRigVMFunction_DrawDebugLibraryBase
+{
+	GENERATED_BODY()
+
+	RIGVM_METHOD()
+	UE_API virtual void Execute() override;
+
+	// Debug Drawer
+	UPROPERTY(meta = (Output))
+	FDebugDrawer Drawer;
+
+	// Input Drawers
+	UPROPERTY(meta = (Input))
+	TArray<FDebugDrawer> Drawers;
+};
+
+// Debug Draws a line
+USTRUCT(meta=(DisplayName="Draw Debug Line"))
+struct FRigUnit_DrawDebugLine : public FRigUnit_DrawDebugLibraryBase
+{
+	GENERATED_BODY()
+
+	RIGVM_METHOD()
+	UE_API virtual void Execute() override;
+
+	// Debug Drawer
+	UPROPERTY(meta = (Input))
+    FDebugDrawer Drawer;
+
+	// Line Start Location
+	UPROPERTY(meta = (Input))
+    FVector StartLocation = FVector::ZeroVector;
+
+	// Line End Location
+	UPROPERTY(meta = (Input))
+    FVector EndLocation = FVector::ZeroVector;
+
+	// Line Style
+	UPROPERTY(meta = (Input))
+    FDrawDebugLineStyle LineStyle;
+
+	// If to use depth testing while drawing
+	UPROPERTY(meta = (Input))
+    bool bDepthTest = true;
+
+	// The execution result
+	UPROPERTY(EditAnywhere, DisplayName = "Execute", Category = "BeginExecution", meta = (Input, Output))
+	FRigVMExecuteContext ExecuteContext;
+};
+
+// Debug Draw a sphere at the given location
+USTRUCT(meta = (DisplayName = "Draw Debug Location"))
+struct FRigUnit_DrawDebugLocation : public FRigUnit_DrawDebugLibraryBase
+{
+	GENERATED_BODY()
+
+	RIGVM_METHOD()
+	UE_API virtual void Execute() override;
+
+	// Debug Drawer
+	UPROPERTY(meta = (Input))
+	FDebugDrawer Drawer;
+
+	// Draw Location
+	UPROPERTY(meta = (Input))
+	FVector Location = FVector::ZeroVector;
+
+	// Line Style
+	UPROPERTY(meta = (Input))
+	FDrawDebugLineStyle LineStyle;
+
+	// If to use depth testing while drawing
+	UPROPERTY(meta = (Input))
+	bool bDepthTest = true;
+
+	// Radius of the sphere to draw
+	UPROPERTY(meta = (Input))
+	float DrawRadius = 10.0f;
+
+	// Number of segments in the sphere
+	UPROPERTY(meta = (Input))
+	int32 Segments = 8;
+
+	// The execution result
+	UPROPERTY(EditAnywhere, DisplayName = "Execute", Category = "BeginExecution", meta = (Input, Output))
+	FRigVMExecuteContext ExecuteContext;
+};
+
+//  Debug Draw a line at the given location scaled by the given velocity
+USTRUCT(meta = (DisplayName = "Draw Debug Velocity"))
+struct FRigUnit_DrawDebugVelocity : public FRigUnit_DrawDebugLibraryBase
+{
+	GENERATED_BODY()
+
+	RIGVM_METHOD()
+	UE_API virtual void Execute() override;
+
+	// Debug Drawer
+	UPROPERTY(meta = (Input))
+	FDebugDrawer Drawer;
+
+	// Draw Location
+	UPROPERTY(meta = (Input))
+	FVector Location = FVector::ZeroVector;
+
+	// Velocity to draw
+	UPROPERTY(meta = (Input))
+	FVector Velocity = FVector::ZeroVector;
+
+	// Line Style
+	UPROPERTY(meta = (Input))
+	FDrawDebugLineStyle LineStyle;
+
+	// If to use depth testing while drawing
+	UPROPERTY(meta = (Input))
+	bool bDepthTest = true;
+
+	// Scale of the line to draw
+	UPROPERTY(meta = (Input))
+	float DrawVelocityLineScale = 1.0f;
+
+	// The execution result
+	UPROPERTY(EditAnywhere, DisplayName = "Execute", Category = "BeginExecution", meta = (Input, Output))
+	FRigVMExecuteContext ExecuteContext;
+};
+
+// Debug Draws a transform trajectory
+USTRUCT(meta = (DisplayName = "Draw Debug Transform Trajectory"))
+struct FRigUnit_DrawDebugTransformTrajectory : public FRigUnit_DrawDebugLibraryBase
+{
+	GENERATED_BODY()
+
+	RIGVM_METHOD()
+	UE_API virtual void Execute() override;
+
+	// Debug Drawer
+	UPROPERTY(meta = (Input))
+	FDebugDrawer Drawer;
+
+	// Transform Trajectory
+	UPROPERTY(meta = (Input))
+	FTransformTrajectory TransformTrajectory;
+
+	// Line Style
+	UPROPERTY(meta = (Input))
+	FDrawDebugLineStyle LineStyle;
+
+	// If to use depth testing while drawing
+	UPROPERTY(meta = (Input))
+	bool bDepthTest = true;
+
+	// Sample sphere radius
+	UPROPERTY(meta = (Input, Units = "cm"))
+	float Radius = 5.0f;
+
+	// Vertical offset to apply to prevent clipping with ground plane
+	UPROPERTY(meta = (Input, Units="cm"))
+	float VerticalOffset = 2.0f;
+
+	// The execution result
+	UPROPERTY(EditAnywhere, DisplayName = "Execute", Category = "BeginExecution", meta = (Input, Output))
+	FRigVMExecuteContext ExecuteContext;
+};
+
+// Debug Draws a string 
+USTRUCT(meta = (DisplayName = "Draw Debug String"))
+struct FRigUnit_DrawDebugString : public FRigUnit_DrawDebugLibraryBase
+{
+	GENERATED_BODY()
+
+	RIGVM_METHOD()
+	UE_API virtual void Execute() override;
+
+	// Debug Drawer
+	UPROPERTY(meta = (Input))
+	FDebugDrawer Drawer;
+
+	// String to draw
+	UPROPERTY(meta = (Input))
+	FString String = TEXT("");
+
+	// Draw Location
+	UPROPERTY(meta = (Input))
+	FVector Location = FVector::ZeroVector;
+
+	// Draw Rotation
+	UPROPERTY(meta = (Input))
+	FRotator Rotation = FRotator::ZeroRotator;
+
+	// Line Style
+	UPROPERTY(meta = (Input))
+	FDrawDebugLineStyle LineStyle;
+
+	// If to use depth testing while drawing
+	UPROPERTY(meta = (Input))
+	bool bDepthTest = true;
+
+	// String Settings
+	UPROPERTY(meta = (Input))
+	FDrawDebugStringSettings Settings;
+
+	// The execution result
+	UPROPERTY(EditAnywhere, DisplayName = "Execute", Category = "BeginExecution", meta = (Input, Output))
+	FRigVMExecuteContext ExecuteContext;
+};
+
+// Debug draws a name 
+USTRUCT(meta = (DisplayName = "Draw Debug Name"))
+struct FRigUnit_DrawDebugName : public FRigUnit_DrawDebugLibraryBase
+{
+	GENERATED_BODY()
+
+	RIGVM_METHOD()
+	UE_API virtual void Execute() override;
+
+	// Debug Drawer
+	UPROPERTY(meta = (Input))
+	FDebugDrawer Drawer;
+
+	// Name to draw
+	UPROPERTY(meta = (Input))
+	FName Name = NAME_None;
+
+	// Draw Location
+	UPROPERTY(meta = (Input))
+	FVector Location = FVector::ZeroVector;
+
+	// Draw Rotation
+	UPROPERTY(meta = (Input))
+	FRotator Rotation = FRotator::ZeroRotator;
+
+	// Line Style
+	UPROPERTY(meta = (Input))
+	FDrawDebugLineStyle LineStyle;
+
+	// If to use depth testing while drawing
+	UPROPERTY(meta = (Input))
+	bool bDepthTest = true;
+
+	// String Settings
+	UPROPERTY(meta = (Input))
+	FDrawDebugStringSettings Settings;
+
+	// The execution result
+	UPROPERTY(EditAnywhere, DisplayName = "Execute", Category = "BeginExecution", meta = (Input, Output))
+	FRigVMExecuteContext ExecuteContext;
+};
+
+#undef UE_API
\ No newline at end of file
