DrawDebugLibrary: Fixed some RigVM initialization errors.

#Synced-CL 48855086

diff --git a/Engine/Plugins/Experimental/DrawDebugLibrary/Source/DrawDebugLibrary/Public/RigUnit_DrawDebugLibrary.h b/Engine/Plugins/Experimental/DrawDebugLibrary/Source/DrawDebugLibrary/Public/RigUnit_DrawDebugLibrary.h
--- a/Engine/Plugins/Experimental/DrawDebugLibrary/Source/DrawDebugLibrary/Public/RigUnit_DrawDebugLibrary.h
+++ b/Engine/Plugins/Experimental/DrawDebugLibrary/Source/DrawDebugLibrary/Public/RigUnit_DrawDebugLibrary.h
@@ -245,11 +245,11 @@ struct FRigVMFunction_DrawDebugLocalOffset : public FRigVMFunction_DrawDebugLibr
 
 	// Output Location
 	UPROPERTY(meta = (Output))
-	FVector Location;
+	FVector Location = FVector::ZeroVector;
 
 	// Output Rotation
 	UPROPERTY(meta = (Output))
-	FRotator Rotation;
+	FRotator Rotation = FRotator::ZeroRotator;
 };
 
 // Returns a new location and rotation after applying a local offset to the input draw location and rotation 
@@ -271,11 +271,11 @@ struct FRigVMFunction_DrawDebugOrientUprightToFloor : public FRigVMFunction_Draw
 
 	// Output Location
 	UPROPERTY(meta = (Output))
-	FVector Location;
+	FVector Location = FVector::ZeroVector;
 
 	// Output Rotation
 	UPROPERTY(meta = (Output))
-	FRotator Rotation;
+	FRotator Rotation = FRotator::ZeroRotator;
 };
 
 // Create a new location and rotation that orients debug drawing of floor objects such as shapes upright
@@ -297,11 +297,11 @@ struct FRigVMFunction_DrawDebugOrientFloorToUpright : public FRigVMFunction_Draw
 
 	// Output Location
 	UPROPERTY(meta = (Output))
-	FVector Location;
+	FVector Location = FVector::ZeroVector;
 
 	// Output Rotation
 	UPROPERTY(meta = (Output))
-	FRotator Rotation;
+	FRotator Rotation = FRotator::ZeroRotator;
 };
 
 // Create a new location and rotation that orients debug drawing of upright objects such as text to face the camera
@@ -331,11 +331,11 @@ struct FRigVMFunction_DrawDebugOrientUprightToCamera : public FRigVMFunction_Dra
 
 	// Output Location
 	UPROPERTY(meta = (Output))
-	FVector Location;
+	FVector Location = FVector::ZeroVector;
 
 	// Output Rotation
 	UPROPERTY(meta = (Output))
-	FRotator Rotation;
+	FRotator Rotation = FRotator::ZeroRotator;
 };
 
 // Create a new location and rotation that orients debug drawing of floor objects such as shapes to face the camera
@@ -365,11 +365,11 @@ struct FRigVMFunction_DrawDebugOrientFloorToCamera : public FRigVMFunction_DrawD
 
 	// Output Location
 	UPROPERTY(meta = (Output))
-	FVector Location;
+	FVector Location = FVector::ZeroVector;
 
 	// Output Rotation
 	UPROPERTY(meta = (Output))
-	FRotator Rotation;
+	FRotator Rotation = FRotator::ZeroRotator;
 };
 
 // Debug Draw a point
@@ -379,7 +379,7 @@ struct FRigUnit_DrawDebugPoint : public FRigUnit_DrawDebugLibraryBase
 	GENERATED_BODY()
 
 	RIGVM_METHOD()
-		UE_API virtual void Execute() override;
+	UE_API virtual void Execute() override;
 
 	// Debug Drawer
 	UPROPERTY(meta = (Input))
