DrawDebugLibrary: Small helper for making linearly spaced values from array view

#Synced-CL 50320943

diff --git a/Engine/Plugins/Experimental/DrawDebugLibrary/Source/DrawDebugLibrary/Private/DrawDebugLibrary.cpp b/Engine/Plugins/Experimental/DrawDebugLibrary/Source/DrawDebugLibrary/Private/DrawDebugLibrary.cpp
--- a/Engine/Plugins/Experimental/DrawDebugLibrary/Source/DrawDebugLibrary/Private/DrawDebugLibrary.cpp
+++ b/Engine/Plugins/Experimental/DrawDebugLibrary/Source/DrawDebugLibrary/Private/DrawDebugLibrary.cpp
@@ -859,6 +859,14 @@ void UDrawDebugLibrary::VisualLoggerLogText(const FDebugDrawer& Drawer, const FT
 void UDrawDebugLibrary::MakeLinearlySpacedFloatArray(TArray<float>& OutValues, const float Start, const float Stop, const int32 Num)
 {
 	OutValues.SetNumUninitialized(Num);
+	MakeLinearlySpacedFloatArrayView(OutValues, Start, Stop);
+}
+
+void UDrawDebugLibrary::MakeLinearlySpacedFloatArrayView(TArrayView<float> OutValues, const float Start, const float Stop)
+{
+	const int32 Num = OutValues.Num();
+
+	if (Num == 0) { return; }
 
 	if (Num == 1)
 	{
diff --git a/Engine/Plugins/Experimental/DrawDebugLibrary/Source/DrawDebugLibrary/Public/DrawDebugLibrary.h b/Engine/Plugins/Experimental/DrawDebugLibrary/Source/DrawDebugLibrary/Public/DrawDebugLibrary.h
--- a/Engine/Plugins/Experimental/DrawDebugLibrary/Source/DrawDebugLibrary/Public/DrawDebugLibrary.h
+++ b/Engine/Plugins/Experimental/DrawDebugLibrary/Source/DrawDebugLibrary/Public/DrawDebugLibrary.h
@@ -673,6 +673,7 @@ public: // Convenience functions
 	/** Convenience function for making an array of linearly spaced float values */
 	UFUNCTION(BlueprintPure = true, Category = "Draw Debug Library")
 	static UE_API void MakeLinearlySpacedFloatArray(TArray<float>& OutValues, const float Start = 0.0f, const float Stop = 1.0f, const int32 Num = 10);
+	static UE_API void MakeLinearlySpacedFloatArrayView(TArrayView<float> OutValues, const float Start = 0.0f, const float Stop = 1.0f);
 
 	/** Convenience function that adds a float to the end of an array, popping values from the front once the max is reached */
 	UFUNCTION(BlueprintPure = false, Category = "Draw Debug Library", meta = (NotBlueprintThreadSafe))
