DrawDebugLibrary: Added function for drawing several names

#Synced-CL 50413561

diff --git a/Engine/Plugins/Experimental/DrawDebugLibrary/Source/DrawDebugLibrary/Private/DrawDebugLibrary.cpp b/Engine/Plugins/Experimental/DrawDebugLibrary/Source/DrawDebugLibrary/Private/DrawDebugLibrary.cpp
--- a/Engine/Plugins/Experimental/DrawDebugLibrary/Source/DrawDebugLibrary/Private/DrawDebugLibrary.cpp
+++ b/Engine/Plugins/Experimental/DrawDebugLibrary/Source/DrawDebugLibrary/Private/DrawDebugLibrary.cpp
@@ -2968,6 +2968,28 @@ void UDrawDebugLibrary::DrawDebugName(const FDebugDrawer& Drawer, const FName& N
 	DrawDebugStringView(Drawer, Name.ToString(), Location, Rotation, LineStyle, bDepthTest, Settings);
 }
 
+void UDrawDebugLibrary::DrawDebugNames(const FDebugDrawer& Drawer, const TArray<FName>& Names, const FVector& Location, const FRotator& Rotation, const FString& Separator, const FDrawDebugLineStyle& LineStyle, const bool bDepthTest, const FDrawDebugStringSettings& Settings)
+{
+	DrawDebugNamesArrayView(Drawer, Names, Location, Rotation, Separator, LineStyle, bDepthTest, Settings);
+}
+
+void UDrawDebugLibrary::DrawDebugNamesArrayView(const FDebugDrawer& Drawer, const TArrayView<const FName> Names, const FVector& Location, const FRotator& Rotation, const FStringView Separator, const FDrawDebugLineStyle& LineStyle, const bool bDepthTest, const FDrawDebugStringSettings& Settings)
+{
+	FString String;
+	const int32 NameNum = Names.Num();
+
+	for (int32 NameIdx = 0; NameIdx < NameNum; NameIdx++)
+	{
+		String += Names[NameIdx].ToString();
+		if (NameIdx != NameNum - 1)
+		{
+			String += Separator;
+		}
+	}
+
+	DrawDebugStringView(Drawer, String, Location, Rotation, LineStyle, bDepthTest, Settings);
+}
+
 void UDrawDebugLibrary::DrawDebugText(const FDebugDrawer& Drawer, const FText& Text, const FVector& Location, const FRotator& Rotation, const FDrawDebugLineStyle& LineStyle, const bool bDepthTest, const FDrawDebugStringSettings& Settings)
 {
 	DrawDebugStringView(Drawer, Text.ToString(), Location, Rotation, LineStyle, bDepthTest, Settings);
diff --git a/Engine/Plugins/Experimental/DrawDebugLibrary/Source/DrawDebugLibrary/Public/DrawDebugLibrary.h b/Engine/Plugins/Experimental/DrawDebugLibrary/Source/DrawDebugLibrary/Public/DrawDebugLibrary.h
--- a/Engine/Plugins/Experimental/DrawDebugLibrary/Source/DrawDebugLibrary/Public/DrawDebugLibrary.h
+++ b/Engine/Plugins/Experimental/DrawDebugLibrary/Source/DrawDebugLibrary/Public/DrawDebugLibrary.h
@@ -1184,6 +1184,11 @@ public: // Text
 	UFUNCTION(BlueprintCallable, Category = "Draw Debug Library", meta = (DevelopmentOnly, AutoCreateRefTerm = "Name, Location, Rotation, LineStyle, Settings"))
 	static UE_API void DrawDebugName(const FDebugDrawer& Drawer, const FName& Name, const FVector& Location, const FRotator& Rotation, const FDrawDebugLineStyle& LineStyle = FDrawDebugLineStyle(), const bool bDepthTest = true, const FDrawDebugStringSettings& Settings = FDrawDebugStringSettings());
 
+	/** Debug Draw several names */
+	UFUNCTION(BlueprintCallable, Category = "Draw Debug Library", meta = (DevelopmentOnly, AutoCreateRefTerm = "Names, Location, Rotation, LineStyle, Settings"))
+	static UE_API void DrawDebugNames(const FDebugDrawer& Drawer, const TArray<FName>& Names, const FVector& Location, const FRotator& Rotation, const FString& Separator = TEXT(", "), const FDrawDebugLineStyle& LineStyle = FDrawDebugLineStyle(), const bool bDepthTest = true, const FDrawDebugStringSettings& Settings = FDrawDebugStringSettings());
+	static UE_API void DrawDebugNamesArrayView(const FDebugDrawer& Drawer, const TArrayView<const FName> Names, const FVector& Location, const FRotator& Rotation, const FStringView Separator = TEXT(", "), const FDrawDebugLineStyle& LineStyle = FDrawDebugLineStyle(), const bool bDepthTest = true, const FDrawDebugStringSettings& Settings = FDrawDebugStringSettings());
+
 	/** Debug Draw some text. Will only render ASCII characters. */
 	UFUNCTION(BlueprintCallable, Category = "Draw Debug Library", meta = (DevelopmentOnly, AutoCreateRefTerm = "Text, Location, Rotation, LineStyle, Settings"))
 	static UE_API void DrawDebugText(const FDebugDrawer& Drawer, const FText& Text, const FVector& Location, const FRotator& Rotation, const FDrawDebugLineStyle& LineStyle = FDrawDebugLineStyle(), const bool bDepthTest = true, const FDrawDebugStringSettings& Settings = FDrawDebugStringSettings());
