DrawDebugLibrary: Added DrawDebugNames RigUnit

#Synced-CL 50487925

diff --git a/Engine/Plugins/Experimental/DrawDebugLibrary/Source/DrawDebugLibrary/Private/RigUnit_DrawDebugLibrary.cpp b/Engine/Plugins/Experimental/DrawDebugLibrary/Source/DrawDebugLibrary/Private/RigUnit_DrawDebugLibrary.cpp
--- a/Engine/Plugins/Experimental/DrawDebugLibrary/Source/DrawDebugLibrary/Private/RigUnit_DrawDebugLibrary.cpp
+++ b/Engine/Plugins/Experimental/DrawDebugLibrary/Source/DrawDebugLibrary/Private/RigUnit_DrawDebugLibrary.cpp
@@ -449,6 +449,11 @@ FRigUnit_DrawDebugName_Execute()
 	UDrawDebugLibrary::DrawDebugName(Drawer, Name, Location, Rotation, LineStyle, bDepthTest, Settings);
 }
 
+FRigUnit_DrawDebugNames_Execute()
+{
+	UDrawDebugLibrary::DrawDebugNamesArrayView(Drawer, Names, Location, Rotation, Separator, LineStyle, bDepthTest, Settings);
+}
+
 FRigUnit_DrawDebugVisualLoggerDrawString_Execute()
 {
 	UDrawDebugLibrary::VisualLoggerDrawStringView(Drawer, String, Location, Color);
diff --git a/Engine/Plugins/Experimental/DrawDebugLibrary/Source/DrawDebugLibrary/Public/RigUnit_DrawDebugLibrary.h b/Engine/Plugins/Experimental/DrawDebugLibrary/Source/DrawDebugLibrary/Public/RigUnit_DrawDebugLibrary.h
--- a/Engine/Plugins/Experimental/DrawDebugLibrary/Source/DrawDebugLibrary/Public/RigUnit_DrawDebugLibrary.h
+++ b/Engine/Plugins/Experimental/DrawDebugLibrary/Source/DrawDebugLibrary/Public/RigUnit_DrawDebugLibrary.h
@@ -3016,6 +3016,48 @@ struct FRigUnit_DrawDebugName : public FRigUnit_DrawDebugLibraryBase
 	FDrawDebugStringSettings Settings;
 };
 
+// Debug draws several names
+USTRUCT(meta = (DisplayName = "Draw Debug Names"))
+struct FRigUnit_DrawDebugNames : public FRigUnit_DrawDebugLibraryBase
+{
+	GENERATED_BODY()
+
+	RIGVM_METHOD()
+	UE_API virtual void Execute() override;
+
+	// Debug Drawer
+	UPROPERTY(meta = (Input))
+	FDebugDrawer Drawer;
+
+	// Names to draw
+	UPROPERTY(meta = (Input))
+	TArray<FName> Names;
+
+	// Draw Location
+	UPROPERTY(meta = (Input))
+	FVector Location = FVector::ZeroVector;
+
+	// Draw Rotation
+	UPROPERTY(meta = (Input))
+	FRotator Rotation = FRotator::ZeroRotator;
+
+	// Separator to put between names
+	UPROPERTY(meta = (Input))
+	FString Separator = TEXT(", ");
+
+	// Line Style
+	UPROPERTY(meta = (Input))
+	FDrawDebugLineStyle LineStyle;
+
+	// If to use depth testing while drawing
+	UPROPERTY(meta = (Input))
+	bool bDepthTest = true;
+
+	// String Settings
+	UPROPERTY(meta = (Input))
+	FDrawDebugStringSettings Settings;
+};
+
 // Debug Draw a string to the visual logger. Will do nothing if not using a VisualLoggerDrawer. Will not display on-screen during recording.
 USTRUCT(meta = (DisplayName = "Visual Logger Draw String"))
 struct FRigUnit_DrawDebugVisualLoggerDrawString : public FRigUnit_DrawDebugLibraryBase
