DrawDebugLibrary: Minor fix for ref parameter

#Synced-CL 51066804

diff --git a/Engine/Plugins/Experimental/DrawDebugLibrary/Source/DrawDebugLibrary/Public/DrawDebugLibrary.h b/Engine/Plugins/Experimental/DrawDebugLibrary/Source/DrawDebugLibrary/Public/DrawDebugLibrary.h
--- a/Engine/Plugins/Experimental/DrawDebugLibrary/Source/DrawDebugLibrary/Public/DrawDebugLibrary.h
+++ b/Engine/Plugins/Experimental/DrawDebugLibrary/Source/DrawDebugLibrary/Public/DrawDebugLibrary.h
@@ -1097,7 +1097,7 @@ public: // Poses / Skeletons
 	static UE_API void DrawDebugBoneLink(const FDebugDrawer& Drawer, const FVector& ChildLocation, const FVector& ParentLocation, const FRotator& ParentRotation, const FLinearColor Color = FLinearColor(0.0f, 0.0f, 0.025f, 1.0f), const bool bDepthTest = true, const float Radius = 5.0f);
 
 	/** Draws a skeleton from a skinned mesh component */
-	UFUNCTION(BlueprintCallable, Category = "Draw Debug Library", meta = (DevelopmentOnly))
+	UFUNCTION(BlueprintCallable, Category = "Draw Debug Library", meta = (DevelopmentOnly, AutoCreateRefTerm = "Settings"))
 	static UE_API void DrawDebugSkeletonFromSkinnedMeshComponent(
 		const FDebugDrawer& Drawer,
 		const USkinnedMeshComponent* SkinnedMeshComponent,
