DrawDebugLibrary: Added functions for drawing bounding boxes

#Synced-CL 51630663

diff --git a/Engine/Plugins/Experimental/DrawDebugLibrary/Source/DrawDebugLibrary/Private/DrawDebugLibrary.cpp b/Engine/Plugins/Experimental/DrawDebugLibrary/Source/DrawDebugLibrary/Private/DrawDebugLibrary.cpp
--- a/Engine/Plugins/Experimental/DrawDebugLibrary/Source/DrawDebugLibrary/Private/DrawDebugLibrary.cpp
+++ b/Engine/Plugins/Experimental/DrawDebugLibrary/Source/DrawDebugLibrary/Private/DrawDebugLibrary.cpp
@@ -13,6 +13,7 @@
 #include "Components/SkinnedMeshComponent.h"
 #include "RigVMCore/RigVMRegistry.h"
 #include "RigVM/Public/RigVMCore/RigVMDrawInterface.h"
+#include "Engine/StaticMesh.h"
 
 #include UE_INLINE_GENERATED_CPP_BY_NAME(DrawDebugLibrary)
 
@@ -1511,6 +1512,39 @@ void UDrawDebugLibrary::DrawDebugBox(const FDebugDrawer& Drawer, const FVector&
 		}, LineStyle, bDepthTest);
 }
 
+void UDrawDebugLibrary::DrawDebugBoxStruct(const FDebugDrawer& Drawer, const FVector& Location, const FRotator& Rotation, const FDrawDebugLineStyle& LineStyle, const bool bDepthTest, const FBox& Box)
+{
+	FVector V[8];
+	Box.GetVertices(V);
+
+	for (int32 Idx = 0; Idx < 8; Idx++)
+	{
+		V[Idx] = Location + Rotation.RotateVector(V[Idx]);
+	}
+
+	DrawDebugLinesPairsArrayView(Drawer, {
+		{ V[4], V[3] },
+		{ V[3], V[0] },
+		{ V[0], V[2] },
+		{ V[2], V[4] },
+
+		{ V[7], V[5] },
+		{ V[5], V[1] },
+		{ V[1], V[6] },
+		{ V[6], V[7] },
+
+		{ V[4], V[7] },
+		{ V[3], V[5] },
+		{ V[0], V[1] },
+		{ V[2], V[6] },
+		}, LineStyle, bDepthTest);
+}
+
+void UDrawDebugLibrary::DrawDebugBoxSphereBounds(const FDebugDrawer& Drawer, const FVector& Location, const FRotator& Rotation, const FDrawDebugLineStyle& LineStyle, const bool bDepthTest, const FBoxSphereBounds& Bounds)
+{
+	DrawDebugBoxStruct(Drawer, Location, Rotation, LineStyle, bDepthTest, FBox::BuildAABB(Bounds.Origin, Bounds.BoxExtent));
+}
+
 void UDrawDebugLibrary::DrawDebugSphere(
 	const FDebugDrawer& Drawer,
 	const FVector& Location,
@@ -2429,6 +2463,13 @@ void UDrawDebugLibrary::DrawDebugPoseFromArrayViews(const FDebugDrawer& Drawer,
 	}
 }
 
+void UDrawDebugLibrary::DrawDebugStaticMeshBoundingBox(const FDebugDrawer& Drawer, const UStaticMesh* StaticMesh, const FVector& Location, const FRotator& Rotation, const FDrawDebugLineStyle& LineStyle, const bool bDepthTest)
+{
+	if (StaticMesh)
+	{
+		DrawDebugBoxSphereBounds(Drawer, Location, Rotation, LineStyle, bDepthTest, StaticMesh->GetBounds());
+	}
+}
 
 void UDrawDebugLibrary::DrawDebugVelocities(const FDebugDrawer& Drawer, const TArray<FVector>& Locations, const TArray<FVector>& Velocities, const FDrawDebugLineStyle& LineStyle, const bool bDepthTest, const float DrawVelocityLineScale)
 {
diff --git a/Engine/Plugins/Experimental/DrawDebugLibrary/Source/DrawDebugLibrary/Public/DrawDebugLibrary.h b/Engine/Plugins/Experimental/DrawDebugLibrary/Source/DrawDebugLibrary/Public/DrawDebugLibrary.h
--- a/Engine/Plugins/Experimental/DrawDebugLibrary/Source/DrawDebugLibrary/Public/DrawDebugLibrary.h
+++ b/Engine/Plugins/Experimental/DrawDebugLibrary/Source/DrawDebugLibrary/Public/DrawDebugLibrary.h
@@ -15,6 +15,7 @@ struct FRigVMDrawInterface;
 struct FTransformTrajectory;
 struct FAnimInstanceProxy;
 class USkinnedMeshComponent;
+class UStaticMesh;
 
 class FDrawDebugLibraryModule : public IModuleInterface
 {
@@ -912,6 +913,14 @@ public: // Basic Shapes
 	UFUNCTION(BlueprintCallable, Category = "Draw Debug Library", meta = (DevelopmentOnly, AutoCreateRefTerm = "Location, Rotation, LineStyle"))
 	static UE_API void DrawDebugBox(const FDebugDrawer& Drawer, const FVector& Location, const FRotator& Rotation, const FDrawDebugLineStyle& LineStyle = FDrawDebugLineStyle(), const bool bDepthTest = true, const FVector HalfExtents = FVector(10.0f, 10.0f, 10.0f));
 
+	/** Debug Draw an FBox */
+	UFUNCTION(BlueprintCallable, Category = "Draw Debug Library", meta = (DevelopmentOnly, AutoCreateRefTerm = "Location, Rotation, Box"))
+	static UE_API void DrawDebugBoxStruct(const FDebugDrawer& Drawer, const FVector& Location, const FRotator& Rotation, const FDrawDebugLineStyle& LineStyle = FDrawDebugLineStyle(), const bool bDepthTest = true, const FBox& Box = FBox());
+
+	/** Debug Draw a FBoxSphereBounds */
+	UFUNCTION(BlueprintCallable, Category = "Draw Debug Library", meta = (DevelopmentOnly, AutoCreateRefTerm = "Location, Rotation, Bounds"))
+	static UE_API void DrawDebugBoxSphereBounds(const FDebugDrawer& Drawer, const FVector& Location, const FRotator& Rotation, const FDrawDebugLineStyle& LineStyle = FDrawDebugLineStyle(), const bool bDepthTest = true, const FBoxSphereBounds& Bounds = FBoxSphereBounds());
+
 	/** Debug Draw a sphere */
 	UFUNCTION(BlueprintCallable, Category = "Draw Debug Library", meta = (DevelopmentOnly, AutoCreateRefTerm = "Location, Rotation, LineStyle"))
 	static UE_API void DrawDebugSphere(const FDebugDrawer& Drawer, const FVector& Location, const FRotator& Rotation, const FDrawDebugLineStyle& LineStyle = FDrawDebugLineStyle(), const bool bDepthTest = true, const float Radius = 10.0f, const int32 Segments = 8);
@@ -1134,6 +1143,13 @@ public: // Poses / Skeletons
 	static UE_API void DrawDebugPose(const FDebugDrawer& Drawer, const TArray<FVector>& BoneLocations, const TArray<FVector>& BoneLinearVelocities, const FTransform& RelativeTransform = FTransform(), const FDrawDebugLineStyle& LineStyle = FDrawDebugLineStyle(), const bool bDepthTest = true, const float DrawVelocityLineScale = 1.0f);
 	static UE_API void DrawDebugPoseFromArrayViews(const FDebugDrawer& Drawer, const TArrayView<const FVector> BoneLocations, const TArrayView<const FVector> BoneLinearVelocities, const FTransform& RelativeTransform = FTransform(), const FDrawDebugLineStyle& LineStyle = FDrawDebugLineStyle(), const bool bDepthTest = true, const float DrawVelocityLineScale = 1.0f);
 
+public: // Objects
+
+	/** Debug Draw the bounding box of a static mesh */
+	UFUNCTION(BlueprintCallable, Category = "Draw Debug Library", meta = (DevelopmentOnly, AutoCreateRefTerm = "Location, Rotation"))
+	static UE_API void DrawDebugStaticMeshBoundingBox(const FDebugDrawer& Drawer, const UStaticMesh* StaticMesh, const FVector& Location, const FRotator& Rotation, const FDrawDebugLineStyle& LineStyle = FDrawDebugLineStyle(), const bool bDepthTest = true);
+
+
 public: // Text
 
 	/** Get the number of line segments required to draw debug a string. */
